[AdvancedRenamer] Moved plugin

diff --git a/Engine/Plugins/Experimental/AdvancedRenamer/AdvancedRenamer.uplugin b/Engine/Plugins/Experimental/AdvancedRenamer/AdvancedRenamer.uplugin
new file mode 100644
--- /dev/null
+++ b/Engine/Plugins/Experimental/AdvancedRenamer/AdvancedRenamer.uplugin
@@ -0,0 +1,23 @@
+{
+	"FileVersion": 3,
+	"Version": 0.1,
+	"VersionName": "0.1.0-alpha",
+	"FriendlyName": "Advanced Renamer",
+	"Description": "Renaming actors and assets",
+	"Category": "Experimental",
+	"CreatedBy": "Epic Games, Inc.",
+	"CreatedByURL": "http://epicgames.com",
+	"DocsURL": "",
+	"MarketplaceURL": "",
+	"SupportURL": "",
+	"CanContainContent": true,
+	"IsBetaVersion": true,
+	"Installed": false,
+	"Modules": [
+		{
+			"Name": "AdvancedRenamer",
+			"Type": "Editor",
+			"LoadingPhase": "Default"
+		}
+	]
+}
diff --git a/Engine/Plugins/Experimental/AdvancedRenamer/Source/AdvancedRenamer/AdvancedRenamer.Build.cs b/Engine/Plugins/Experimental/AdvancedRenamer/Source/AdvancedRenamer/AdvancedRenamer.Build.cs
new file mode 100644
--- /dev/null
+++ b/Engine/Plugins/Experimental/AdvancedRenamer/Source/AdvancedRenamer/AdvancedRenamer.Build.cs
@@ -0,0 +1,35 @@
+// Copyright Epic Games, Inc. All Rights Reserved.
+
+using UnrealBuildTool;
+
+public class AdvancedRenamer : ModuleRules
+{
+	public AdvancedRenamer(ReadOnlyTargetRules Target) : base(Target)
+	{
+		PCHUsage = ModuleRules.PCHUsageMode.UseExplicitOrSharedPCHs;
+
+		PublicDependencyModuleNames.AddRange(
+			new string[]
+			{
+				"CoreUObject",
+			}
+		);
+
+		PrivateDependencyModuleNames.AddRange(
+			new string[]
+			{
+				"ContentBrowser",
+				"Core",
+				"Engine",
+				"InputCore",
+				"LevelEditor",
+				"Projects",
+				"Slate",
+				"SlateCore",
+				"ToolMenus",
+				"UMG",
+				"UnrealEd"
+			}
+		);
+	}
+}
diff --git a/Engine/Plugins/Experimental/AdvancedRenamer/Source/AdvancedRenamer/Private/AdvancedRenamerBlueprintLibrary.cpp b/Engine/Plugins/Experimental/AdvancedRenamer/Source/AdvancedRenamer/Private/AdvancedRenamerBlueprintLibrary.cpp
new file mode 100644
--- /dev/null
+++ b/Engine/Plugins/Experimental/AdvancedRenamer/Source/AdvancedRenamer/Private/AdvancedRenamerBlueprintLibrary.cpp
@@ -0,0 +1,207 @@
+// Copyright Epic Games, Inc. All Rights Reserved.
+
+#include "AdvancedRenamerBlueprintLibrary.h"
+#include "Blueprint/UserWidget.h"
+#include "EngineUtils.h"
+#include "Framework/Application/SlateApplication.h"
+#include "Providers/AdvancedRenamerActorProvider.h"
+#include "Slate/SAdvancedRenamerPanel.h"
+#include "Templates/SharedPointer.h"
+#include "Toolkits/IToolkitHost.h"
+#include "Widgets/SWidget.h"
+#include "Widgets/SWindow.h"
+
+#define LOCTEXT_NAMESPACE "AdvancedRenamerBlueprintLibrary"
+
+namespace UE::AdvancedRenamer::Private
+{
+	TSharedRef<SWindow> CreateAdvancedRenamerWindow()
+	{
+		return
+			SNew(SWindow)
+			.Title(LOCTEXT("AdvancedRenameWindow", "Rename Actors"))
+			.ClientSize(FVector2D(600.0f, 500.0f))
+			.SizingRule(ESizingRule::FixedSize)
+			.SupportsMaximize(false)
+			.SupportsMinimize(false);
+	}
+}
+
+void UAdvancedRenamerBlueprintLibrary::BP_OpenAdvancedRenamer(UObject* RenameProvider, UUserWidget* UserWidget)
+{
+	TSharedPtr<SWidget> ParentWidget = nullptr;
+	
+	if (UserWidget)
+	{
+		ParentWidget = UserWidget->TakeWidget();
+	}
+
+	using namespace UE::AdvancedRenamer::Private;
+
+	TSharedRef<SWindow> AdvancedRenameWindow = CreateAdvancedRenamerWindow();
+	AdvancedRenameWindow->SetContent(SNew(SAdvancedRenamerPanel).ObjectProvider(RenameProvider));
+
+	TSharedPtr<SWidget> ParentWindow = FSlateApplication::Get().FindBestParentWindowForDialogs(ParentWidget);
+	FSlateApplication::Get().AddModalWindow(AdvancedRenameWindow, ParentWindow);
+}
+
+void UAdvancedRenamerBlueprintLibrary::OpenAdvancedRenamer(TSharedRef<IAdvancedRenamerProvider> RenameProvider,
+	TSharedPtr<IToolkitHost> ToolkitHost)
+{
+	if (ToolkitHost.IsValid())
+	{
+		OpenAdvancedRenamer(RenameProvider, ToolkitHost->GetParentWidget());
+	}
+}
+
+void UAdvancedRenamerBlueprintLibrary::OpenAdvancedRenamer(TSharedRef<IAdvancedRenamerProvider> RenameProvider, TSharedPtr<SWidget> InParentWidget)
+{
+	using namespace UE::AdvancedRenamer::Private;
+
+	TSharedRef<SWindow> AdvancedRenameWindow = CreateAdvancedRenamerWindow();
+	AdvancedRenameWindow->SetContent(SNew(SAdvancedRenamerPanel).SharedProvider(RenameProvider));
+
+	TSharedPtr<SWidget> ParentWindow = FSlateApplication::Get().FindBestParentWindowForDialogs(InParentWidget);
+	FSlateApplication::Get().AddModalWindow(AdvancedRenameWindow, ParentWindow);
+}
+
+void UAdvancedRenamerBlueprintLibrary::BP_OpenAdvancedRenamerForActors(const TArray<AActor*>& Actors, UUserWidget* UserWidget)
+{
+	TSharedPtr<SWidget> ParentWidget = nullptr;
+
+	if (UserWidget)
+	{
+		ParentWidget = UserWidget->TakeWidget();
+	}
+
+	OpenAdvancedRenamerForActors(Actors, ParentWidget);
+}
+
+void UAdvancedRenamerBlueprintLibrary::OpenAdvancedRenamerForActors(const TArray<AActor*>& Actors, TSharedPtr<IToolkitHost> ToolkitHost)
+{
+	TSharedPtr<SWidget> ParentWidget = nullptr;
+
+	if (ToolkitHost.IsValid())
+	{
+		ParentWidget = ToolkitHost->GetParentWidget();
+	}
+
+	OpenAdvancedRenamerForActors(Actors, ParentWidget);
+}
+
+void UAdvancedRenamerBlueprintLibrary::OpenAdvancedRenamerForActors(const TArray<AActor*>& Actors, TSharedPtr<SWidget> ParentWidget)
+{
+	TArray<TWeakObjectPtr<AActor>> WeakObjects;
+
+	WeakObjects.Reserve(Actors.Num());
+
+	Algo::Transform(
+		Actors,
+		WeakObjects,
+		[](AActor* InActor)
+		{
+			return TWeakObjectPtr<AActor>(InActor);
+		}
+	);
+
+	TSharedRef<FAdvancedRenamerActorProvider> ObjectProvider = MakeShared<FAdvancedRenamerActorProvider>();
+	ObjectProvider->SetActorList(WeakObjects);
+
+	OpenAdvancedRenamer(ObjectProvider, ParentWidget);
+}
+
+TArray<AActor*> UAdvancedRenamerBlueprintLibrary::GetActorsSharingClassesInWorld(const TArray<AActor*>& InActors)
+{
+	TSet<UClass*> SelectedClasses;
+	bool bHasActorClass = false;
+	UWorld* World = nullptr;
+
+	// Scan selected items and add valid classes to the selected classes list.
+	for (AActor* SelectedActor : InActors)
+	{
+		if (!IsValid(SelectedActor))
+		{
+			continue;
+		}
+
+		if (!World)
+		{
+			World = SelectedActor->GetWorld();
+
+			if (!World)
+			{
+				break;
+			}
+		}
+
+		UClass* ActorClass = SelectedActor->GetClass();
+
+		/**
+		 * If we have a default AActor selected then all actors in the world share a
+		 * class with the selected actors. We don't need anything other than the AActor
+		 * class to get matches. Empty the array, store that and move on.
+		 */
+		if (ActorClass == AActor::StaticClass())
+		{
+			bHasActorClass = true;
+			SelectedClasses.Empty();
+			break;
+		}
+
+		SelectedClasses.Add(ActorClass);
+	}
+
+	if (!World)
+	{
+		return InActors;
+	}
+
+	TArray<UClass*> NonInheritingActorClasses;
+
+	if (bHasActorClass)
+	{
+		NonInheritingActorClasses.Add(AActor::StaticClass());
+	}
+	else
+	{
+		for (UClass* ActorClass : SelectedClasses)
+		{
+			bool bFoundParent = false;
+
+			for (UClass* ActorClassCheck : SelectedClasses)
+			{
+				if (ActorClass == ActorClassCheck)
+				{
+					continue;
+				}
+
+				if (ActorClass->IsChildOf(ActorClassCheck))
+				{
+					bFoundParent = true;
+					break;
+				}
+			}
+
+			if (!bFoundParent)
+			{
+				NonInheritingActorClasses.Add(ActorClass);
+			}
+		}
+	}
+
+	// Create outliner items for all the items matching the class list that are renameable.
+	TArray<AActor*> AllActors;
+	AllActors.Reserve(InActors.Num());
+
+	for (UClass* ActorClass : NonInheritingActorClasses)
+	{
+		for (AActor* Actor : TActorRange<AActor>(World, ActorClass))
+		{
+			AllActors.Add(Actor);
+		}
+	}
+
+	return AllActors;
+}
+
+#undef LOCTEXT_NAMESPACE
diff --git a/Engine/Plugins/Experimental/AdvancedRenamer/Source/AdvancedRenamer/Private/AdvancedRenamerCommands.cpp b/Engine/Plugins/Experimental/AdvancedRenamer/Source/AdvancedRenamer/Private/AdvancedRenamerCommands.cpp
new file mode 100644
--- /dev/null
+++ b/Engine/Plugins/Experimental/AdvancedRenamer/Source/AdvancedRenamer/Private/AdvancedRenamerCommands.cpp
@@ -0,0 +1,36 @@
+// Copyright Epic Games, Inc. All Rights Reserved.
+
+#include "AdvancedRenamerCommands.h"
+#include "AdvancedRenamerStyle.h"
+#include "Framework/Commands/InputChord.h"
+#include "Framework/Commands/UICommandInfo.h"
+#include "Styling/ISlateStyle.h"
+
+#define LOCTEXT_NAMESPACE "AdvancedRenamerCommands"
+
+FAdvancedRenamerCommands::FAdvancedRenamerCommands()
+	: TCommands<FAdvancedRenamerCommands>(
+		TEXT("AdvancedRenamer"),
+		LOCTEXT("AdvancedRenamer", "Advanced Renamer"),
+		NAME_None,
+		FAdvancedRenamerStyle::Get().GetStyleSetName()
+	)
+{
+}
+
+void FAdvancedRenamerCommands::RegisterCommands()
+{
+	UI_COMMAND(RenameSelectedActors
+		, "Rename Selected Actors"
+		, "Opens the Advanced Renamer Panel to rename all selected actors."
+		, EUserInterfaceActionType::Button
+		, FInputChord())
+
+	UI_COMMAND(RenameSharedClassActors
+		, "Rename Actors of Selected Actor Classes"
+		, "Opens the Advanced Renamer Panel to rename all actors sharing a class with any selected actor."
+		, EUserInterfaceActionType::Button
+		, FInputChord())
+}
+
+#undef LOCTEXT_NAMESPACE
diff --git a/Engine/Plugins/Experimental/AdvancedRenamer/Source/AdvancedRenamer/Private/AdvancedRenamerCommands.h b/Engine/Plugins/Experimental/AdvancedRenamer/Source/AdvancedRenamer/Private/AdvancedRenamerCommands.h
new file mode 100644
--- /dev/null
+++ b/Engine/Plugins/Experimental/AdvancedRenamer/Source/AdvancedRenamer/Private/AdvancedRenamerCommands.h
@@ -0,0 +1,19 @@
+// Copyright Epic Games, Inc. All Rights Reserved.
+
+#pragma once
+
+#include "Framework/Commands/Commands.h"
+#include "Templates/SharedPointer.h"
+
+class FUICommandInfo;
+
+class FAdvancedRenamerCommands : public TCommands<FAdvancedRenamerCommands>
+{
+public:
+	FAdvancedRenamerCommands();
+
+	virtual void RegisterCommands() override;
+
+	TSharedPtr<FUICommandInfo> RenameSelectedActors;
+	TSharedPtr<FUICommandInfo> RenameSharedClassActors;
+};
diff --git a/Engine/Plugins/Experimental/AdvancedRenamer/Source/AdvancedRenamer/Private/AdvancedRenamerModule.cpp b/Engine/Plugins/Experimental/AdvancedRenamer/Source/AdvancedRenamer/Private/AdvancedRenamerModule.cpp
new file mode 100644
--- /dev/null
+++ b/Engine/Plugins/Experimental/AdvancedRenamer/Source/AdvancedRenamer/Private/AdvancedRenamerModule.cpp
@@ -0,0 +1,34 @@
+// Copyright Epic Games, Inc. All Rights Reserved.
+
+#include "AdvancedRenamerModule.h"
+#include "AdvancedRenamerCommands.h"
+#include "AdvancedRenamerStyle.h"
+#include "Integrations/AdvancedRenamerContentBrowserIntegration.h"
+#include "Integrations/AdvancedRenamerLevelEditorIntegration.h"
+#include "Modules/ModuleManager.h"
+
+DEFINE_LOG_CATEGORY(LogARP);
+
+FAdvancedRenamerModule& FAdvancedRenamerModule::Get()
+{
+	static const FName ModuleName = TEXT("AdvancedRenamer");
+	return FModuleManager::LoadModuleChecked<FAdvancedRenamerModule>(ModuleName);
+}
+
+void FAdvancedRenamerModule::StartupModule()
+{
+	FAdvancedRenamerStyle::Initialize();
+	FAdvancedRenamerCommands::Register();
+	FAdvancedRenamerContentBrowserIntegration::Initialize();
+	FAdvancedRenamerLevelEditorIntegration::Initialize();
+}
+
+void FAdvancedRenamerModule::ShutdownModule()
+{
+	FAdvancedRenamerCommands::Unregister();
+	FAdvancedRenamerStyle::Shutdown();
+	FAdvancedRenamerContentBrowserIntegration::Shutdown();
+	FAdvancedRenamerLevelEditorIntegration::Shutdown();
+}
+
+IMPLEMENT_MODULE(FAdvancedRenamerModule, AdvancedRenamer)
diff --git a/Engine/Plugins/Experimental/AdvancedRenamer/Source/AdvancedRenamer/Private/AdvancedRenamerStyle.cpp b/Engine/Plugins/Experimental/AdvancedRenamer/Source/AdvancedRenamer/Private/AdvancedRenamerStyle.cpp
new file mode 100644
--- /dev/null
+++ b/Engine/Plugins/Experimental/AdvancedRenamer/Source/AdvancedRenamer/Private/AdvancedRenamerStyle.cpp
@@ -0,0 +1,82 @@
+// Copyright Epic Games, Inc. All Rights Reserved.
+
+#include "AdvancedRenamerStyle.h"
+#include "Brushes/SlateImageBrush.h"
+#include "Interfaces/IPluginManager.h"
+#include "Misc/Paths.h"
+#include "Styling/AppStyle.h"
+#include "Styling/ISlateStyle.h"
+#include "Styling/SlateStyle.h"
+#include "Styling/SlateStyleRegistry.h"
+#include "Styling/SlateTypes.h"
+#include "Styling/StyleColors.h"
+
+TSharedPtr<FSlateStyleSet> FAdvancedRenamerStyle::StyleInstance = nullptr;
+
+void FAdvancedRenamerStyle::Initialize()
+{
+	if (!StyleInstance.IsValid())
+	{
+		StyleInstance = Create();
+		FSlateStyleRegistry::RegisterSlateStyle(*StyleInstance);
+	}
+}
+
+void FAdvancedRenamerStyle::Shutdown()
+{
+	if (StyleInstance.IsValid())
+	{
+		FSlateStyleRegistry::UnRegisterSlateStyle(*StyleInstance);
+		ensure(StyleInstance.IsUnique());
+		StyleInstance.Reset();
+	}
+}
+
+FName FAdvancedRenamerStyle::GetStyleSetName()
+{
+	static FName StyleSetName(TEXT("AdvancedRenamerStyle"));
+	return StyleSetName;
+}
+
+TSharedRef<FSlateStyleSet> FAdvancedRenamerStyle::Create()
+{
+	TSharedRef<FSlateStyleSet> Style = MakeShared<FSlateStyleSet>("AdvancedRenamer");
+
+	TSharedPtr<IPlugin> Plugin = IPluginManager::Get().FindPlugin(TEXT("AdvancedRenamer"));
+	check(Plugin.IsValid());
+
+	if (Plugin.IsValid())
+	{
+		Style->SetContentRoot(FPaths::Combine(Plugin->GetBaseDir(), TEXT("Resources")));
+	}
+
+	Style->Set("AdvancedRenamer.Image.Radio.BlackBackground", new FSlateVectorImageBrush(
+		FPaths::Combine(Style->GetContentRootDir(), "/Icons/radio-background.svg"), 
+		FVector2D(16.f, 16.f), FStyleColors::Foldout.GetSpecifiedColor()
+	));
+
+	FCheckBoxStyle BlackRadioButton = FAppStyle::GetWidgetStyle<FCheckBoxStyle>("RadioButton");
+	BlackRadioButton.SetBackgroundImage(*Style->GetBrush("AdvancedRenamer.Image.Radio.BlackBackground"));
+
+	Style->Set("AdvancedRenamer.Style.BlackRadioButton", BlackRadioButton);
+
+	const FButtonStyle DarkButton = FButtonStyle(FAppStyle::Get().GetWidgetStyle<FButtonStyle>("SimpleButton"))
+		.SetNormal(FSlateColorBrush(FStyleColors::Recessed))
+		.SetHovered(FSlateColorBrush(FStyleColors::Hover))
+		.SetPressed(FSlateColorBrush(FStyleColors::Header))
+		.SetDisabled(FSlateColorBrush(FStyleColors::Dropdown));
+
+	Style->Set("AdvancedRenamer.Style.DarkButton", DarkButton);
+
+	return Style;
+}
+
+const ISlateStyle& FAdvancedRenamerStyle::Get()
+{
+	if (!StyleInstance.IsValid())
+	{
+		Initialize();
+	}
+
+	return *StyleInstance;
+}
diff --git a/Engine/Plugins/Experimental/AdvancedRenamer/Source/AdvancedRenamer/Private/AdvancedRenamerStyle.h b/Engine/Plugins/Experimental/AdvancedRenamer/Source/AdvancedRenamer/Private/AdvancedRenamerStyle.h
new file mode 100644
--- /dev/null
+++ b/Engine/Plugins/Experimental/AdvancedRenamer/Source/AdvancedRenamer/Private/AdvancedRenamerStyle.h
@@ -0,0 +1,28 @@
+// Copyright Epic Games, Inc. All Rights Reserved.
+
+#pragma once
+
+#include "Templates/SharedPointer.h"
+
+class FSlateStyleSet;
+class ISlateStyle;
+class FName;
+
+class FAdvancedRenamerStyle
+{
+public:
+
+	static void Initialize();
+
+	static void Shutdown();
+
+	static const ISlateStyle& Get();
+
+	static FName GetStyleSetName();
+
+private:
+
+	static TSharedRef<FSlateStyleSet> Create();
+	static TSharedPtr<FSlateStyleSet> StyleInstance;
+
+};
diff --git a/Engine/Plugins/Experimental/AdvancedRenamer/Source/AdvancedRenamer/Private/Integrations/AdvancedRenamerContentBrowserIntegration.cpp b/Engine/Plugins/Experimental/AdvancedRenamer/Source/AdvancedRenamer/Private/Integrations/AdvancedRenamerContentBrowserIntegration.cpp
new file mode 100644
--- /dev/null
+++ b/Engine/Plugins/Experimental/AdvancedRenamer/Source/AdvancedRenamer/Private/Integrations/AdvancedRenamerContentBrowserIntegration.cpp
@@ -0,0 +1,89 @@
+// Copyright Epic Games, Inc. All Rights Reserved.
+
+#include "Integrations/AdvancedRenamerContentBrowserIntegration.h"
+#include "AdvancedRenamerBlueprintLibrary.h"
+#include "ContentBrowserDelegates.h"
+#include "ContentBrowserModule.h"
+#include "Delegates/IDelegateInstance.h"
+#include "Framework/Commands/UIAction.h"
+#include "Framework/MultiBox/MultiBoxBuilder.h"
+#include "Framework/MultiBox/MultiBoxExtender.h"
+#include "Providers/AdvancedRenamerAssetProvider.h"
+
+#define LOCTEXT_NAMESPACE "AdvancedRenamerContentBrowserIntegration"
+
+namespace UE::AdvancedRenamer::Private
+{
+	FDelegateHandle ContentBrowserDelegateHandle;
+
+	void OpenAdvancedRenamer(const TArray<FAssetData> AssetArray)
+	{
+		TSharedRef<FAdvancedRenamerAssetProvider> AssetProvider = MakeShared<FAdvancedRenamerAssetProvider>();
+		AssetProvider->SetAssetList(AssetArray);
+
+		TSharedPtr<SWidget> HostWidget = nullptr;
+
+		UAdvancedRenamerBlueprintLibrary::OpenAdvancedRenamer(StaticCastSharedRef<IAdvancedRenamerProvider>(AssetProvider), HostWidget);
+	}
+
+	void ExtendAssetMenu(FMenuBuilder& MenuBuilder, const TArray<FAssetData> AssetArray)
+	{
+		MenuBuilder.AddMenuEntry(
+			LOCTEXT("AdvancedRename", "Advanced Rename"),
+			LOCTEXT("AdvancedRenameTooltip", "Opens the Advanced Renamer Panel to rename all selected assets."),
+			FSlateIcon(FAppStyle::GetAppStyleSetName(), "ContentBrowser.AssetActions.Rename"),
+			FUIAction(FExecuteAction::CreateStatic(&OpenAdvancedRenamer, AssetArray))
+		);
+	}
+
+	TSharedRef<FExtender> CreateContentMenuExtender(const TArray<FAssetData>& AssetArray)
+	{
+		TSharedRef<FExtender> Extender = MakeShared<FExtender>();
+
+		if (AssetArray.Num() > 0)
+		{
+			Extender->AddMenuExtension(
+				"CommonAssetActions",
+				EExtensionHook::Position::After,
+				nullptr,
+				FMenuExtensionDelegate::CreateStatic(&ExtendAssetMenu, AssetArray)
+			);
+		}
+
+		return Extender;
+	}
+}
+
+void FAdvancedRenamerContentBrowserIntegration::Initialize()
+{
+	using namespace UE::AdvancedRenamer::Private;
+
+	// Register Content Browser selection extensions
+	FContentBrowserModule& ContentBrowserModule = FModuleManager::LoadModuleChecked<FContentBrowserModule>("ContentBrowser");
+	TArray<FContentBrowserMenuExtender_SelectedAssets>& AssetContextMenuDelegates = ContentBrowserModule.GetAllAssetViewContextMenuExtenders();
+	AssetContextMenuDelegates.Add(FContentBrowserMenuExtender_SelectedAssets::CreateStatic(&CreateContentMenuExtender));
+	ContentBrowserDelegateHandle = AssetContextMenuDelegates.Last().GetHandle();
+}
+
+void FAdvancedRenamerContentBrowserIntegration::Shutdown()
+{
+	using namespace UE::AdvancedRenamer::Private;
+
+	if (ContentBrowserDelegateHandle.IsValid())
+	{
+		FContentBrowserModule* ContentBrowserModule = FModuleManager::GetModulePtr<FContentBrowserModule>("ContentBrowser");
+
+		if (ContentBrowserModule)
+		{
+			TArray<FContentBrowserMenuExtender_SelectedAssets>& AssetContextMenuDelegates = ContentBrowserModule->GetAllAssetViewContextMenuExtenders();
+			AssetContextMenuDelegates.RemoveAll([DelegateHandle = ContentBrowserDelegateHandle](const FContentBrowserMenuExtender_SelectedAssets& Delegate)
+				{
+					return Delegate.GetHandle() == DelegateHandle;
+				});
+		}
+
+		ContentBrowserDelegateHandle.Reset();
+	}
+}
+
+#undef LOCTEXT_NAMESPACE
diff --git a/Engine/Plugins/Experimental/AdvancedRenamer/Source/AdvancedRenamer/Private/Integrations/AdvancedRenamerContentBrowserIntegration.h b/Engine/Plugins/Experimental/AdvancedRenamer/Source/AdvancedRenamer/Private/Integrations/AdvancedRenamerContentBrowserIntegration.h
new file mode 100644
--- /dev/null
+++ b/Engine/Plugins/Experimental/AdvancedRenamer/Source/AdvancedRenamer/Private/Integrations/AdvancedRenamerContentBrowserIntegration.h
@@ -0,0 +1,10 @@
+// Copyright Epic Games, Inc. All Rights Reserved.
+
+#pragma once
+
+class FAdvancedRenamerContentBrowserIntegration
+{
+public:
+	static void Initialize();
+	static void Shutdown();
+};
diff --git a/Engine/Plugins/Experimental/AdvancedRenamer/Source/AdvancedRenamer/Private/Integrations/AdvancedRenamerLevelEditorIntegration.cpp b/Engine/Plugins/Experimental/AdvancedRenamer/Source/AdvancedRenamer/Private/Integrations/AdvancedRenamerLevelEditorIntegration.cpp
new file mode 100644
--- /dev/null
+++ b/Engine/Plugins/Experimental/AdvancedRenamer/Source/AdvancedRenamer/Private/Integrations/AdvancedRenamerLevelEditorIntegration.cpp
@@ -0,0 +1,149 @@
+// Copyright Epic Games, Inc. All Rights Reserved.
+
+#include "Integrations/AdvancedRenamerLevelEditorIntegration.h"
+#include "AdvancedRenamerBlueprintLibrary.h"
+#include "AdvancedRenamerCommands.h"
+#include "Containers/Array.h"
+#include "Delegates/IDelegateInstance.h"
+#include "EditorModeManager.h"
+#include "GameFramework/Actor.h"
+#include "ILevelEditor.h"
+#include "LevelEditor.h"
+#include "Selection.h"
+#include "Templates/SharedPointer.h"
+
+namespace UE::AdvancedRenamer::Private
+{
+	FDelegateHandle LevelEditorCreatedDelegateHandle;
+
+	TArray<AActor*> GetSelectedActors(TSharedRef<ILevelEditor> InLevelEditor)
+	{
+		USelection* ActorSelection = InLevelEditor->GetEditorModeManager().GetSelectedActors();
+
+		if (!ActorSelection)
+		{
+			return {};
+		}
+
+		TArray<AActor*> SelectedActors;
+		int32 Count = ActorSelection->GetSelectedObjects<AActor>(SelectedActors);
+
+		return SelectedActors;
+	}
+
+	bool CanOpenAdvancedRenamer(TWeakPtr<ILevelEditor> InLevelEditorWeak)
+	{
+		TSharedPtr<ILevelEditor> LevelEditor = InLevelEditorWeak.Pin();
+
+		if (!LevelEditor)
+		{
+			return false;
+		}
+
+		TArray<AActor*> SelectedActors = GetSelectedActors(LevelEditor.ToSharedRef());
+
+		for (AActor* SelectedActor : SelectedActors)
+		{
+			if (IsValid(SelectedActor))
+			{
+				return true;
+			}
+		}
+
+		return false;
+	}
+
+	void RenameSelectedActors(TWeakPtr<ILevelEditor> InLevelEditorWeak)
+	{
+		TSharedPtr<ILevelEditor> LevelEditor = InLevelEditorWeak.Pin();
+
+		if (!LevelEditor)
+		{
+			return;
+		}
+
+		TArray<AActor*> SelectedActors = GetSelectedActors(LevelEditor.ToSharedRef());
+
+		UAdvancedRenamerBlueprintLibrary::OpenAdvancedRenamerForActors(SelectedActors, StaticCastSharedPtr<IToolkitHost>(LevelEditor));
+	}
+
+	void RenameSharedClassActors(TWeakPtr<ILevelEditor> InLevelEditorWeak)
+	{
+		TSharedPtr<ILevelEditor> LevelEditor = InLevelEditorWeak.Pin();
+
+		if (!LevelEditor)
+		{
+			return;
+		}
+
+		TArray<AActor*> SelectedActors = GetSelectedActors(LevelEditor.ToSharedRef());
+		SelectedActors = UAdvancedRenamerBlueprintLibrary::GetActorsSharingClassesInWorld(SelectedActors);
+
+		UAdvancedRenamerBlueprintLibrary::OpenAdvancedRenamerForActors(SelectedActors, StaticCastSharedPtr<IToolkitHost>(LevelEditor));
+	}
+
+	void OnLevelEditorCreated(TSharedPtr<ILevelEditor> InLevelEditor)
+	{
+		const FAdvancedRenamerCommands& AdvRenCommands = FAdvancedRenamerCommands::Get();
+
+		InLevelEditor->GetLevelEditorActions()->MapAction(
+			AdvRenCommands.RenameSelectedActors,
+			FExecuteAction::CreateStatic(&RenameSelectedActors, InLevelEditor.ToWeakPtr()),
+			FCanExecuteAction::CreateStatic(&CanOpenAdvancedRenamer, InLevelEditor.ToWeakPtr())
+		);
+
+		InLevelEditor->GetLevelEditorActions()->MapAction(
+			AdvRenCommands.RenameSharedClassActors,
+			FExecuteAction::CreateStatic(&RenameSharedClassActors, InLevelEditor.ToWeakPtr()),
+			FCanExecuteAction::CreateStatic(&CanOpenAdvancedRenamer, InLevelEditor.ToWeakPtr())
+		);
+	}
+}
+
+void FAdvancedRenamerLevelEditorIntegration::Initialize()
+{
+	using namespace UE::AdvancedRenamer::Private;
+
+	FLevelEditorModule& LevelEditorModule = FModuleManager::LoadModuleChecked<FLevelEditorModule>("LevelEditor");
+	LevelEditorCreatedDelegateHandle = LevelEditorModule.OnLevelEditorCreated().AddStatic(&OnLevelEditorCreated);
+
+	const FAdvancedRenamerCommands& AdvRenCommands = FAdvancedRenamerCommands::Get();
+
+	static const TArray<FName, TInlineAllocator<2>> Menus = {
+		"LevelEditor.ActorContextMenu.EditSubMenu",
+		"LevelEditor.LevelEditorSceneOutliner.ContextMenu.EditSubMenu"
+	};
+
+	const TAttribute<FText> TextAttribute;
+	const FSlateIcon RenameIcon = FSlateIcon(FAppStyle::GetAppStyleSetName(), "GenericCommands.Rename");
+
+	for (const FName& Menu : Menus)
+	{
+		UToolMenu* ToolMenu = UToolMenus::Get()->ExtendMenu(Menu);
+		FToolMenuSection& Section = ToolMenu->FindOrAddSection(NAME_None);
+
+		Section.AddMenuEntry(
+			AdvRenCommands.RenameSelectedActors,
+			TextAttribute,
+			TextAttribute,
+			RenameIcon
+		);
+
+		Section.AddMenuEntry(
+			AdvRenCommands.RenameSharedClassActors,
+			TextAttribute,
+			TextAttribute,
+			RenameIcon
+		);
+	}
+}
+
+void FAdvancedRenamerLevelEditorIntegration::Shutdown()
+{
+	using namespace UE::AdvancedRenamer::Private;
+
+	if (FLevelEditorModule* LevelEditorModule = FModuleManager::LoadModulePtr<FLevelEditorModule>("LevelEditor"))
+	{
+		LevelEditorModule->OnLevelEditorCreated().Remove(LevelEditorCreatedDelegateHandle);
+	}
+}
diff --git a/Engine/Plugins/Experimental/AdvancedRenamer/Source/AdvancedRenamer/Private/Integrations/AdvancedRenamerLevelEditorIntegration.h b/Engine/Plugins/Experimental/AdvancedRenamer/Source/AdvancedRenamer/Private/Integrations/AdvancedRenamerLevelEditorIntegration.h
new file mode 100644
--- /dev/null
+++ b/Engine/Plugins/Experimental/AdvancedRenamer/Source/AdvancedRenamer/Private/Integrations/AdvancedRenamerLevelEditorIntegration.h
@@ -0,0 +1,10 @@
+// Copyright Epic Games, Inc. All Rights Reserved.
+
+#pragma once
+
+class FAdvancedRenamerLevelEditorIntegration
+{
+public:
+	static void Initialize();
+	static void Shutdown();	
+};
diff --git a/Engine/Plugins/Experimental/AdvancedRenamer/Source/AdvancedRenamer/Private/Providers/AdvancedRenamerActorProvider.cpp b/Engine/Plugins/Experimental/AdvancedRenamer/Source/AdvancedRenamer/Private/Providers/AdvancedRenamerActorProvider.cpp
new file mode 100644
--- /dev/null
+++ b/Engine/Plugins/Experimental/AdvancedRenamer/Source/AdvancedRenamer/Private/Providers/AdvancedRenamerActorProvider.cpp
@@ -0,0 +1,110 @@
+// Copyright Epic Games, Inc. All Rights Reserved.
+
+#include "Providers/AdvancedRenamerActorProvider.h"
+#include "GameFramework/Actor.h"
+
+FAdvancedRenamerActorProvider::FAdvancedRenamerActorProvider()
+{
+}
+
+FAdvancedRenamerActorProvider::~FAdvancedRenamerActorProvider()
+{
+}
+
+void FAdvancedRenamerActorProvider::SetActorList(const TArray<TWeakObjectPtr<AActor>>& InActorList)
+{
+	ActorList.Empty();
+	ActorList.Append(InActorList);
+}
+
+void FAdvancedRenamerActorProvider::AddActorList(const TArray<TWeakObjectPtr<AActor>>& InActorList)
+{
+	ActorList.Append(InActorList);
+}
+
+void FAdvancedRenamerActorProvider::AddActorData(AActor* InActor)
+{
+	ActorList.Add(InActor);
+}
+
+AActor* FAdvancedRenamerActorProvider::GetActor(int32 InIndex) const
+{
+	if (!ActorList.IsValidIndex(InIndex))
+	{
+		return nullptr;
+	}
+
+	return ActorList[InIndex].Get();
+}
+
+int32 FAdvancedRenamerActorProvider::Num() const
+{
+	return ActorList.Num();
+}
+
+bool FAdvancedRenamerActorProvider::IsValidIndex(int32 InIndex) const
+{
+	AActor* Actor = GetActor(InIndex);
+
+	return IsValid(Actor);
+}
+
+FString FAdvancedRenamerActorProvider::GetOriginalName(int32 InIndex) const
+{
+	AActor* Actor = GetActor(InIndex);
+
+	if (!IsValid(Actor))
+	{
+		return "";
+	}
+
+	return Actor->GetActorNameOrLabel();
+}
+
+uint32 FAdvancedRenamerActorProvider::GetHash(int32 InIndex) const
+{
+	AActor* Actor = GetActor(InIndex);
+
+	if (!IsValid(Actor))
+	{
+		return 0;
+	}
+
+	return GetTypeHash(Actor);
+}
+
+bool FAdvancedRenamerActorProvider::RemoveIndex(int32 InIndex)
+{
+	if (!ActorList.IsValidIndex(InIndex))
+	{
+		return false;
+	}
+
+	ActorList.RemoveAt(InIndex);
+	return true;
+}
+
+bool FAdvancedRenamerActorProvider::CanRename(int32 InIndex) const
+{
+	AActor* Actor = GetActor(InIndex);
+
+	if (!IsValid(Actor))
+	{
+		return false;
+	}
+
+	return true;
+}
+
+bool FAdvancedRenamerActorProvider::ExecuteRename(int32 InIndex, const FString& InNewName)
+{
+	AActor* Actor = GetActor(InIndex);
+
+	if (!IsValid(Actor))
+	{
+		return false;
+	}
+
+	Actor->SetActorLabel(InNewName, /* bMarkDirty */ true);
+	return true;
+}
diff --git a/Engine/Plugins/Experimental/AdvancedRenamer/Source/AdvancedRenamer/Private/Providers/AdvancedRenamerAssetProvider.cpp b/Engine/Plugins/Experimental/AdvancedRenamer/Source/AdvancedRenamer/Private/Providers/AdvancedRenamerAssetProvider.cpp
new file mode 100644
--- /dev/null
+++ b/Engine/Plugins/Experimental/AdvancedRenamer/Source/AdvancedRenamer/Private/Providers/AdvancedRenamerAssetProvider.cpp
@@ -0,0 +1,118 @@
+// Copyright Epic Games, Inc. All Rights Reserved.
+
+#include "Providers/AdvancedRenamerAssetProvider.h"
+#include "AssetToolsModule.h"
+#include "IAssetTools.h"
+
+#define LOCTEXT_NAMESPACE "AdvancedRenamerAssetProvider"
+
+FAdvancedRenamerAssetProvider::FAdvancedRenamerAssetProvider()
+{
+}
+
+FAdvancedRenamerAssetProvider::~FAdvancedRenamerAssetProvider()
+{
+}
+
+void FAdvancedRenamerAssetProvider::SetAssetList(const TArray<FAssetData>& InAssetList)
+{
+	AssetList.Empty();
+	AssetList.Append(InAssetList);
+}
+
+void FAdvancedRenamerAssetProvider::AddAssetList(const TArray<FAssetData>& InAssetList)
+{
+	AssetList.Append(InAssetList);
+}
+
+void FAdvancedRenamerAssetProvider::AddAssetData(const FAssetData& InAsset)
+{
+	AssetList.Add(InAsset);
+}
+
+UObject* FAdvancedRenamerAssetProvider::GetAsset(int32 Index) const
+{
+	if (!AssetList.IsValidIndex(Index))
+	{
+		return nullptr;
+	}
+
+	return AssetList[Index].GetAsset();
+}
+
+int32 FAdvancedRenamerAssetProvider::Num() const
+{
+	return AssetList.Num();
+}
+
+bool FAdvancedRenamerAssetProvider::IsValidIndex(int32 Index) const
+{
+	UObject* Asset = GetAsset(Index);
+
+	return IsValid(Asset) && Asset->IsAsset();
+}
+
+FString FAdvancedRenamerAssetProvider::GetOriginalName(int32 Index) const
+{
+	UObject* Asset = GetAsset(Index);
+
+	if (!IsValid(Asset))
+	{
+		return "";
+	}
+
+	return Asset->GetName();
+}
+
+uint32 FAdvancedRenamerAssetProvider::GetHash(int32 Index) const
+{
+	UObject* Asset = GetAsset(Index);
+
+	if (!IsValid(Asset))
+	{
+		return 0;
+	}
+
+	return GetTypeHash(Asset);
+}
+
+bool FAdvancedRenamerAssetProvider::RemoveIndex(int32 Index)
+{
+	if (!AssetList.IsValidIndex(Index))
+	{
+		return false;
+	}
+
+	AssetList.RemoveAt(Index);
+	return true;
+}
+
+bool FAdvancedRenamerAssetProvider::CanRename(int32 Index) const
+{
+	UObject* Asset = GetAsset(Index);
+
+	if (!IsValid(Asset))
+	{
+		return false;
+	}
+
+	return true;
+}
+
+bool FAdvancedRenamerAssetProvider::ExecuteRename(int32 Index, const FString& NewName)
+{
+	UObject* Asset = GetAsset(Index);
+
+	if (!IsValid(Asset))
+	{
+		return false;
+	}
+
+	FString PackagePath = Asset->GetPathName();
+	PackagePath = FPaths::GetPath(PackagePath);
+
+	IAssetTools& AssetTools = FModuleManager::GetModuleChecked<FAssetToolsModule>("AssetTools").Get();
+	return AssetTools.RenameAssets({FAssetRenameData(Asset, PackagePath, NewName)});
+}
+
+#undef LOCTEXT_NAMESPACE
diff --git a/Engine/Plugins/Experimental/AdvancedRenamer/Source/AdvancedRenamer/Private/Providers/AdvancedRenamerObjectAssetProvider.cpp b/Engine/Plugins/Experimental/AdvancedRenamer/Source/AdvancedRenamer/Private/Providers/AdvancedRenamerObjectAssetProvider.cpp
new file mode 100644
--- /dev/null
+++ b/Engine/Plugins/Experimental/AdvancedRenamer/Source/AdvancedRenamer/Private/Providers/AdvancedRenamerObjectAssetProvider.cpp
@@ -0,0 +1,64 @@
+// Copyright Epic Games, Inc. All Rights Reserved.
+
+#include "Providers/AdvancedRenamerObjectAssetProvider.h"
+#include "AssetRegistry/AssetData.h"
+
+void UAdvancedRenamerAssetProvider::BP_SetAssetList_Implementation(const TArray<UObject*>& InAssetList)
+{
+	AssetList.Empty();
+
+	BP_SetAssetList(InAssetList);
+}
+
+void UAdvancedRenamerAssetProvider::BP_AddAssetList_Implementation(const TArray<UObject*>& InAssetList)
+{
+	for (UObject* Object : InAssetList)
+	{
+		BP_AddAssetData(Object);
+	}
+}
+
+void UAdvancedRenamerAssetProvider::BP_AddAssetData_Implementation(UObject* InAsset)
+{
+	AddAssetData(FAssetData(InAsset, true));
+}
+
+UObject* UAdvancedRenamerAssetProvider::BP_GetAsset_Implementation(int32 Index) const
+{
+	return GetAsset(Index);
+}
+
+int32 UAdvancedRenamerAssetProvider::Num() const
+{
+	return BP_Num();
+}
+
+bool UAdvancedRenamerAssetProvider::ExecuteRename(int32 Index, const FString& NewName)
+{
+	return BP_ExecuteRename(Index, NewName);
+}
+
+bool UAdvancedRenamerAssetProvider::CanRename(int32 Index) const
+{
+	return BP_CanRename(Index);
+}
+
+bool UAdvancedRenamerAssetProvider::RemoveIndex(int32 Index)
+{
+	return BP_RemoveIndex(Index);
+}
+
+FString UAdvancedRenamerAssetProvider::GetOriginalName(int32 Index) const
+{
+	return BP_GetOriginalName(Index);
+}
+
+uint32 UAdvancedRenamerAssetProvider::GetHash(int32 Index) const
+{
+	return static_cast<uint32>(BP_GetHash(Index));
+}
+
+bool UAdvancedRenamerAssetProvider::IsValidIndex(int32 Index) const
+{
+	return BP_IsValidIndex(Index);
+}
diff --git a/Engine/Plugins/Experimental/AdvancedRenamer/Source/AdvancedRenamer/Private/Providers/AdvancedRenamerObjectObjectProvider.cpp b/Engine/Plugins/Experimental/AdvancedRenamer/Source/AdvancedRenamer/Private/Providers/AdvancedRenamerObjectObjectProvider.cpp
new file mode 100644
--- /dev/null
+++ b/Engine/Plugins/Experimental/AdvancedRenamer/Source/AdvancedRenamer/Private/Providers/AdvancedRenamerObjectObjectProvider.cpp
@@ -0,0 +1,63 @@
+// Copyright Epic Games, Inc. All Rights Reserved.
+
+#include "Providers/AdvancedRenamerObjectObjectProvider.h"
+
+void UAdvancedRenamerObjectObjectProvider::BP_SetObjectList_Implementation(const TArray<UObject*>& InObjectList)
+{
+	ObjectList.Empty();
+
+	BP_AddObjectList(InObjectList);
+}
+
+void UAdvancedRenamerObjectObjectProvider::BP_AddObjectList_Implementation(const TArray<UObject*>& InObjectList)
+{
+	for (UObject* Object : InObjectList)
+	{
+		BP_AddObjectData(Object);
+	}
+}
+
+void UAdvancedRenamerObjectObjectProvider::BP_AddObjectData_Implementation(UObject* InObject)
+{
+	AddObjectData(InObject);
+}
+
+UObject* UAdvancedRenamerObjectObjectProvider::BP_GetObject_Implementation(int32 Index) const
+{
+	return GetObject(Index);
+}
+
+int32 UAdvancedRenamerObjectObjectProvider::Num() const
+{
+	return BP_Num();
+}
+
+bool UAdvancedRenamerObjectObjectProvider::ExecuteRename(int32 Index, const FString& NewName)
+{
+	return BP_ExecuteRename(Index, NewName);
+}
+
+bool UAdvancedRenamerObjectObjectProvider::CanRename(int32 Index) const
+{
+	return BP_CanRename(Index);
+}
+
+bool UAdvancedRenamerObjectObjectProvider::RemoveIndex(int32 Index)
+{
+	return BP_RemoveIndex(Index);
+}
+
+FString UAdvancedRenamerObjectObjectProvider::GetOriginalName(int32 Index) const
+{
+	return BP_GetOriginalName(Index);
+}
+
+uint32 UAdvancedRenamerObjectObjectProvider::GetHash(int32 Index) const
+{
+	return static_cast<uint32>(BP_GetHash(Index));
+}
+
+bool UAdvancedRenamerObjectObjectProvider::IsValidIndex(int32 Index) const
+{
+	return BP_IsValidIndex(Index);
+}
diff --git a/Engine/Plugins/Experimental/AdvancedRenamer/Source/AdvancedRenamer/Private/Providers/AdvancedRenamerObjectProvider.cpp b/Engine/Plugins/Experimental/AdvancedRenamer/Source/AdvancedRenamer/Private/Providers/AdvancedRenamerObjectProvider.cpp
new file mode 100644
--- /dev/null
+++ b/Engine/Plugins/Experimental/AdvancedRenamer/Source/AdvancedRenamer/Private/Providers/AdvancedRenamerObjectProvider.cpp
@@ -0,0 +1,115 @@
+// Copyright Epic Games, Inc. All Rights Reserved.
+
+#include "Providers/AdvancedRenamerObjectProvider.h"
+#include "UObject/Object.h"
+
+FAdvancedRenamerObjectProvider::FAdvancedRenamerObjectProvider()
+{
+}
+
+FAdvancedRenamerObjectProvider::~FAdvancedRenamerObjectProvider()
+{
+}
+
+void FAdvancedRenamerObjectProvider::SetObjectList(const TArray<TWeakObjectPtr<UObject>>& InObjectList)
+{
+	ObjectList.Empty();
+	ObjectList.Append(InObjectList);
+}
+
+void FAdvancedRenamerObjectProvider::AddObjectList(const TArray<TWeakObjectPtr<UObject>>& InObjectList)
+{
+	ObjectList.Append(InObjectList);
+}
+
+void FAdvancedRenamerObjectProvider::AddObjectData(UObject* InObject)
+{
+	ObjectList.Add(InObject);
+}
+
+UObject* FAdvancedRenamerObjectProvider::GetObject(int32 Index) const
+{
+	if (!ObjectList.IsValidIndex(Index))
+	{
+		return nullptr;
+	}
+
+	return ObjectList[Index].Get();
+}
+
+int32 FAdvancedRenamerObjectProvider::Num() const
+{
+	return ObjectList.Num();
+}
+
+bool FAdvancedRenamerObjectProvider::IsValidIndex(int32 Index) const
+{
+	UObject* Object = GetObject(Index);
+
+	return IsValid(Object);
+}
+
+FString FAdvancedRenamerObjectProvider::GetOriginalName(int32 Index) const
+{
+	UObject* Object = GetObject(Index);
+
+	if (!IsValid(Object))
+	{
+		return "";
+	}
+
+	return Object->GetName();
+}
+
+uint32 FAdvancedRenamerObjectProvider::GetHash(int32 Index) const
+{
+	UObject* Object = GetObject(Index);
+
+	if (!IsValid(Object))
+	{
+		return 0;
+	}
+
+	return GetTypeHash(Object);
+}
+
+bool FAdvancedRenamerObjectProvider::RemoveIndex(int32 Index)
+{
+	if (!ObjectList.IsValidIndex(Index))
+	{
+		return false;
+	}
+
+	ObjectList.RemoveAt(Index);
+	return true;
+}
+
+bool FAdvancedRenamerObjectProvider::CanRename(int32 Index) const
+{
+	UObject* Object = GetObject(Index);
+
+	if (!IsValid(Object))
+	{
+		return false;
+	}
+
+	return true;
+}
+
+bool FAdvancedRenamerObjectProvider::ExecuteRename(int32 Index, const FString& NewName)
+{
+	UObject* Object = GetObject(Index);
+
+	if (!IsValid(Object))
+	{
+		return false;
+	}
+
+	if (Object->Rename(*NewName, nullptr, REN_Test))
+	{
+		Object->Rename(*NewName);
+		return true;
+	}
+
+	return false;
+}
diff --git a/Engine/Plugins/Experimental/AdvancedRenamer/Source/AdvancedRenamer/Private/Slate/SAdvancedRenamerPanel.cpp b/Engine/Plugins/Experimental/AdvancedRenamer/Source/AdvancedRenamer/Private/Slate/SAdvancedRenamerPanel.cpp
new file mode 100644
--- /dev/null
+++ b/Engine/Plugins/Experimental/AdvancedRenamer/Source/AdvancedRenamer/Private/Slate/SAdvancedRenamerPanel.cpp
@@ -0,0 +1,1871 @@
+// Copyright Epic Games, Inc. All Rights Reserved.
+
+#include "Slate/SAdvancedRenamerPanel.h"
+#include "AdvancedRenamerModule.h"
+#include "AdvancedRenamerStyle.h"
+#include "Framework/Commands/GenericCommands.h"
+#include "Framework/MultiBox/MultiBoxBuilder.h"
+#include "Internationalization/Regex.h"
+#include "Providers/AdvancedRenamerObjectProviderInterface.h"
+#include "Providers/AdvancedRenamerProviderInterface.h"
+#include "ScopedTransaction.h"
+#include "Styling/AppStyle.h"
+#include "Styling/StyleColors.h"
+#include "Widgets/Colors/SColorBlock.h"
+#include "Widgets/Input/SButton.h"
+#include "Widgets/Input/SCheckBox.h"
+#include "Widgets/Input/SEditableTextBox.h"
+#include "Widgets/Input/SMultiLineEditableTextBox.h"
+#include "Widgets/Layout/SBorder.h"
+#include "Widgets/Layout/SBox.h"
+#include "Widgets/Input/SSpinBox.h"
+#include "Widgets/Layout/SScrollBox.h"
+#include "Widgets/SCanvas.h"
+#include "Widgets/SNullWidget.h"
+#include "Widgets/Views/SListView.h"
+#include "Widgets/Views/STableRow.h"
+
+#define LOCTEXT_NAMESPACE "SAdvancedRenamerPanel"
+
+FName FAdvancedRenamerPreviewListItem::OriginalNameColumnName = "OriginalName";
+FName FAdvancedRenamerPreviewListItem::NewNameColumnName = "NewName";
+
+namespace UE::AdvancedRenamer::Private
+{
+	const FVector2D WindowSize = {600.f, 500.f};
+	const FSlateFontInfo TitleFont = FCoreStyle::GetDefaultFontStyle("Regular", 12);
+	const FSlateFontInfo RegularFont = FCoreStyle::GetDefaultFontStyle("Regular", 10);
+	const float LeftBlockWidth = 304.f;
+	const FVector2D TitleSize = {150.f, 20.f};
+	const float TitleOffsetY = 25.f;
+	const FVector2D CheckboxOffset = {0.f, 2.f};
+	const FVector2D CheckboxSize = {18.f, 18.f};
+	const FVector2D RadioOffset = {1.f, 2.f};
+	const float LabelStart = 25.f;
+	const float LabelOffsetY = 3.f;
+	const FVector2D LabelSize = {125.f, 18.f};
+	const float EntryStart = 145.f;
+	const FVector2D EntrySize = {155.f, 21.f};
+	const float LineHeight = 26.f;
+	const FVector2D SeparatorSize = {32.f, 21.f};
+	const FVector2D SpinSize1 = {32.f, 21.f};
+	const FVector2D SpinSize2 = {39.f, 21.f};
+	const FVector2D SpinSize3 = {46.f, 21.f};
+	const float RightBlockOffsetX = 314.f;
+	const FVector2D RightBlockSize = {281.f, 490.f};
+	const FVector2D ListViewSize = {277.f, 431.f};
+	const float ListLineHeight = 15.f;
+	const float ApplyButtonHeight = 25.f;
+}
+
+void SAdvancedRenamerPreviewListRow::Construct(const FArguments& InArgs, TSharedPtr<SAdvancedRenamerPanel> InRenamePanel, 
+	const TSharedRef<STableViewBase>& InOwnerTableView, FObjectRenamePreviewListItemPtr InRowItem)
+{
+	RenamePanel = InRenamePanel;
+	RowItem = InRowItem;
+
+	SMultiColumnTableRow<FObjectRenamePreviewListItemPtr>::Construct(FSuperRowType::FArguments(), InOwnerTableView);
+	SetBorderImage(TAttribute<const FSlateBrush*>(this, &SAdvancedRenamerPreviewListRow::GetBorder));
+}
+
+TSharedRef<SWidget> SAdvancedRenamerPreviewListRow::GenerateWidgetForColumn(const FName& ColumnName)
+{
+	if (ColumnName != FAdvancedRenamerPreviewListItem::OriginalNameColumnName
+		&& ColumnName != FAdvancedRenamerPreviewListItem::NewNameColumnName)
+	{
+		return SNullWidget::NullWidget;
+	}
+
+	TSharedPtr<SAdvancedRenamerPanel> RenamePanelSP = RenamePanel.Pin();
+
+	if (!RenamePanelSP.IsValid())
+	{
+		return SNullWidget::NullWidget;
+	}
+
+	FObjectRenamePreviewListItemPtr RowItemSP = RowItem.Pin();
+
+	if (!RowItemSP.IsValid())
+	{
+		return SNullWidget::NullWidget;
+	}
+
+	using namespace UE::AdvancedRenamer::Private;
+
+	TSharedRef<STextBlock> TextBlock = SNew(STextBlock)
+		.ColorAndOpacity(FStyleColors::AccentWhite.GetSpecifiedColor())
+		.Font(RegularFont)
+		.Margin(FMargin(5.f, 0.f, 5.f, 0.f));
+		
+	if (ColumnName == FAdvancedRenamerPreviewListItem::OriginalNameColumnName)
+	{
+		TextBlock->SetText(FText::FromString(RowItemSP->OriginalName));
+		RenamePanelSP->MinDesiredOriginalNameWidth = FMath::Max(RenamePanelSP->MinDesiredOriginalNameWidth, TextBlock->ComputeDesiredSize(1.f).X);
+	}
+	else if (ColumnName == FAdvancedRenamerPreviewListItem::NewNameColumnName)
+	{
+		TextBlock->SetText(FText::FromString(RowItemSP->NewName));
+		RenamePanelSP->MinDesiredNewNameWidth = FMath::Max(RenamePanelSP->MinDesiredOriginalNameWidth, TextBlock->ComputeDesiredSize(1.f).X);
+	}
+
+	RenamePanelSP->UpdateRequiredListWidth();
+
+	return TextBlock;
+}
+
+void SAdvancedRenamerPanel::Construct(const FArguments& InArgs)
+{
+	if (InArgs._ObjectProvider != nullptr)
+	{
+		ObjectProvider = TStrongObjectPtr<UObject>(InArgs._ObjectProvider);
+		check(ObjectProvider->GetClass()->ImplementsInterface(UAdvancedRenamerObjectProvider::StaticClass()));
+	}
+	else if (InArgs._SharedProvider.IsValid())
+	{
+		SharedProvider = InArgs._SharedProvider;
+	}
+	else
+	{
+		checkNoEntry();
+	}
+
+	int32 Count = Num();
+	check(Count > 0);
+
+	for (int32 Index = 0; Index < Count; ++Index)
+	{
+		if (!CanRename(Index))
+		{
+			RemoveIndex(Index);
+			--Index;
+			--Count;
+			continue;
+		}
+
+		int32 Hash = GetHash(Index);
+		FString OriginalName = GetOriginalName(Index);
+
+		ListData.Add(MakeShared<FAdvancedRenamerPreviewListItem>(Hash, OriginalName));
+	}
+
+	CommandList = MakeShared<FUICommandList>();
+	CommandList->MapAction(
+		FGenericCommands::Get().Delete,
+		FExecuteAction::CreateSP(this, &SAdvancedRenamerPanel::RemoveSelectedObjects),
+		FCanExecuteAction()
+	);
+
+	using namespace UE::AdvancedRenamer::Private;
+
+	// @formatter:off
+	TSharedRef<SCanvas> Canvas = SNew(SCanvas)
+		+ SCanvas::Slot()
+		.Position(FVector2D(0.f, 0.f))
+		.Size(WindowSize)
+		[
+			SNew(SColorBlock)
+			.Color(FStyleColors::Background.GetSpecifiedColor())
+		];
+	// @formatter:on
+
+	CreateLeftPane(Canvas);
+	CreateRightPane(Canvas);
+
+	// @formatter:off
+	ChildSlot
+	[
+		Canvas
+	];
+	// @formatter:on
+}
+
+void SAdvancedRenamerPanel::CreateLeftPane(TSharedRef<SCanvas> Canvas)
+{	
+	using namespace UE::AdvancedRenamer::Private;
+
+	// @formatter:off
+	Canvas->AddSlot()
+		.Position(FVector2D(5.f, 5.f))
+		.Size(FVector2D(LeftBlockWidth, 50.f))
+		[
+			CreateBaseName()
+		];
+
+	Canvas->AddSlot()
+		.Position(FVector2D(5.f, 60.f))
+		.Size(FVector2D(LeftBlockWidth, 102.f))
+		[
+			CreatePrefix()
+		];
+
+	Canvas->AddSlot()
+		.Position(FVector2D(5.f, 167.f))
+		.Size(FVector2D(LeftBlockWidth, 154.f))
+		[
+			CreateSuffix()
+		];
+
+	Canvas->AddSlot()
+		.Position(FVector2D(5.f, 327.f))
+		.Size(FVector2D(LeftBlockWidth, 168.f))
+		[
+			CreateSearchAndReplace()
+		];
+	// @formatter:on
+}
+
+TSharedRef<SWidget> SAdvancedRenamerPanel::CreateBaseName()
+{
+	using namespace UE::AdvancedRenamer::Private;
+
+	// @formatter:off
+	return SNew(SBorder)
+		.BorderImage(FAppStyle::Get().GetBrush("DetailsView.CategoryTop"))
+		.Content()
+		[
+			SNew(SCanvas)
+			+ SCanvas::Slot()
+			.Position(FVector2D(0.f, LabelOffsetY))
+			.Size(TitleSize)
+			[
+				SNew(STextBlock)
+				.ColorAndOpacity(FStyleColors::AccentBlue.GetSpecifiedColor())
+				.Font(TitleFont)
+				.Text(LOCTEXT("BaseNameTitle", "Name"))
+			]
+			+ SCanvas::Slot()
+			.Position(FVector2D(0.f, TitleOffsetY) + CheckboxOffset)
+			.Size(CheckboxSize)
+			[
+				SAssignNew(BaseNameCheckBox, SCheckBox)
+				.IsChecked(this, &SAdvancedRenamerPanel::IsBaseNameChecked)
+				.OnCheckStateChanged(this, &SAdvancedRenamerPanel::OnBaseNameCheckBoxChanged)
+			]
+			+ SCanvas::Slot()
+			.Position(FVector2D(LabelStart, TitleOffsetY + LabelOffsetY))
+			.Size(LabelSize)
+			[
+				SNew(STextBlock)
+				.ColorAndOpacity(FStyleColors::AccentWhite.GetSpecifiedColor())
+				.Font(RegularFont)
+				.Text(LOCTEXT("BaseName", "Base Name"))
+			]
+			+ SCanvas::Slot()
+			.Position(FVector2D(EntryStart, TitleOffsetY))
+			.Size(EntrySize)
+			[
+				SAssignNew(BaseNameTextBox, SEditableTextBox)
+				.BackgroundColor(FStyleColors::Background.GetSpecifiedColor())
+				.ForegroundColor(FStyleColors::AccentWhite.GetSpecifiedColor())
+				.Font(RegularFont)
+				.HintText(LOCTEXT("BaseNameHint", "Base name"))
+				.IsEnabled(false)
+				.OnTextChanged(this, &SAdvancedRenamerPanel::OnBaseNameChanged)
+			]
+		];
+	// @formatter:on
+}
+
+TSharedRef<SWidget> SAdvancedRenamerPanel::CreatePrefix()
+{
+	using namespace UE::AdvancedRenamer::Private;
+
+	// @formatter:off
+	return SNew(SBorder)
+		.BorderImage(FAppStyle::Get().GetBrush("DetailsView.CategoryTop"))
+		.Content()
+		[
+			SNew(SCanvas)
+
+			// Title
+			+ SCanvas::Slot()
+			.Position(FVector2D(0.f, LabelOffsetY))
+			.Size(TitleSize)
+			[
+				SNew(STextBlock)
+				.ColorAndOpacity(FStyleColors::AccentBlue.GetSpecifiedColor())
+				.Font(TitleFont)
+				.Text(LOCTEXT("PrefixTitle", "Prefix"))
+			]
+
+			// Prefix name
+			+ SCanvas::Slot()
+			.Position(FVector2D(0.f, TitleOffsetY) + CheckboxOffset)
+			.Size(CheckboxSize)
+			[
+				SAssignNew(PrefixCheckBox, SCheckBox)
+				.IsChecked(this, &SAdvancedRenamerPanel::IsPrefixChecked)
+				.OnCheckStateChanged(this, &SAdvancedRenamerPanel::OnPrefixCheckBoxChanged)
+			]
+			+ SCanvas::Slot()
+			.Position(FVector2D(LabelStart, TitleOffsetY + LabelOffsetY))
+			.Size(LabelSize)
+			[
+				SNew(STextBlock)
+				.ColorAndOpacity(FStyleColors::AccentWhite.GetSpecifiedColor())
+				.Font(RegularFont)
+				.Text(LOCTEXT("Prefix", "Prefix"))
+			]
+			+ SCanvas::Slot()
+			.Position(FVector2D(EntryStart, TitleOffsetY))
+			.Size(EntrySize)
+			[
+				SAssignNew(PrefixTextBox, SEditableTextBox)
+				.BackgroundColor(FStyleColors::Background.GetSpecifiedColor())
+				.ForegroundColor(FStyleColors::AccentWhite.GetSpecifiedColor())
+				.Font(RegularFont)
+				.HintText(LOCTEXT("PrefixHint", "Prefix"))
+				.IsEnabled(false)
+				.OnTextChanged(this, &SAdvancedRenamerPanel::OnPrefixChanged)
+			]
+
+			// Remove prefix
+			+ SCanvas::Slot()
+			.Position(FVector2D(LabelStart, TitleOffsetY + LineHeight) + RadioOffset)
+			.Size(CheckboxSize)
+			[
+				SAssignNew(PrefixRemoveCheckBox, SCheckBox)
+				.Style(&FAdvancedRenamerStyle::Get().GetWidgetStyle<FCheckBoxStyle>("AdvancedRenamer.Style.BlackRadioButton"))
+				.IsChecked(this, &SAdvancedRenamerPanel::IsPrefixRemoveChecked)
+				.OnCheckStateChanged(this, &SAdvancedRenamerPanel::OnPrefixRemoveCheckBoxChanged)
+				.IsEnabled(false)
+			]
+			+ SCanvas::Slot()
+			.Position(FVector2D(LabelStart * 2.f, TitleOffsetY + LabelOffsetY + LineHeight))
+			.Size(LabelSize)
+			[
+				SNew(STextBlock)
+				.ColorAndOpacity(FStyleColors::AccentWhite.GetSpecifiedColor())
+				.Font(RegularFont)
+				.Text(LOCTEXT("PrefixRemove", "Remove Prefix"))
+			]
+			+ SCanvas::Slot()
+			.Position(FVector2D(EntryStart, TitleOffsetY + LineHeight))
+			.Size(SeparatorSize)
+			[
+				SAssignNew(PrefixSeparatorTextBox, SEditableTextBox)
+				.BackgroundColor(FStyleColors::Background.GetSpecifiedColor())
+				.ForegroundColor(FStyleColors::AccentWhite.GetSpecifiedColor())
+				.Font(RegularFont)
+				.IsEnabled(false)
+				.Text(LOCTEXT("Underscore", "_"))
+				.OnVerifyTextChanged(this, &SAdvancedRenamerPanel::OnPrefixSeparatorVerifyTextChanged)
+				.OnTextChanged(this, &SAdvancedRenamerPanel::OnPrefixSeparatorChanged)
+			]
+			+ SCanvas::Slot()
+			.Position(FVector2D(EntryStart + SeparatorSize.X + 5.f, TitleOffsetY + LabelOffsetY + LineHeight))
+			.Size(LabelSize)
+			[
+				SNew(STextBlock)
+				.ColorAndOpacity(FStyleColors::AccentWhite.GetSpecifiedColor())
+				.Font(RegularFont)
+				.Text(LOCTEXT("Separator", "(separator)"))
+			]
+
+			// Remove first
+			+ SCanvas::Slot()
+			.Position(FVector2D(LabelStart, TitleOffsetY + LineHeight * 2.f) + RadioOffset)
+			.Size(CheckboxSize)
+			[
+				SAssignNew(PrefixRemoveCharactersCheckBox, SCheckBox)
+				.Style(&FAdvancedRenamerStyle::Get().GetWidgetStyle<FCheckBoxStyle>("AdvancedRenamer.Style.BlackRadioButton"))
+				.IsChecked(this, &SAdvancedRenamerPanel::IsPrefixRemoveCharactersChecked)
+				.OnCheckStateChanged(this, &SAdvancedRenamerPanel::OnPrefixRemoveCharactersCheckBoxChanged)
+				.IsEnabled(false)
+			]
+			+ SCanvas::Slot()
+			.Position(FVector2D(LabelStart * 2.f, TitleOffsetY + LabelOffsetY + LineHeight * 2.f))
+			.Size(LabelSize)
+			[
+				SNew(STextBlock)
+				.ColorAndOpacity(FStyleColors::AccentWhite.GetSpecifiedColor())
+				.Font(RegularFont)
+				.Text(LOCTEXT("PrefixRemoveChars", "Remove First"))
+			]
+			+ SCanvas::Slot()
+			.Position(FVector2D(EntryStart, TitleOffsetY + LineHeight * 2.f))
+			.Size(SpinSize1)
+			[
+
+				SAssignNew(PrefixRemoveCharactersSpinBox, SSpinBox<uint8>)
+				.Style(&FAppStyle::Get(), "Menu.SpinBox")
+				.Font(RegularFont)
+				.MinValue(1)
+				.MaxValue(9)
+				.Value(1)
+				.IsEnabled(false)
+				.OnValueChanged(this, &SAdvancedRenamerPanel::OnPrefixRemoveCharactersChanged)
+			]
+			+ SCanvas::Slot()
+			.Position(FVector2D(EntryStart + SeparatorSize.X + 5.f, TitleOffsetY + LabelOffsetY + LineHeight * 2.f))
+			.Size(LabelSize)
+			[
+				SNew(STextBlock)
+				.ColorAndOpacity(FStyleColors::AccentWhite.GetSpecifiedColor())
+				.Font(RegularFont)
+				.Text(LOCTEXT("Characters", "character(s)"))
+			]
+		];
+	// @formatter:on
+}
+
+TSharedRef<SWidget> SAdvancedRenamerPanel::CreateSuffix()
+{
+	using namespace UE::AdvancedRenamer::Private;
+
+	// @formatter:off
+	return SNew(SBorder)
+		.BorderImage(FAppStyle::Get().GetBrush("DetailsView.CategoryTop"))
+		.Content()
+		[
+			SNew(SCanvas)
+
+			// Title
+			+ SCanvas::Slot()
+			.Position(FVector2D(0.f, LabelOffsetY))
+			.Size(TitleSize)
+			[
+				SNew(STextBlock)
+				.ColorAndOpacity(FStyleColors::AccentBlue.GetSpecifiedColor())
+				.Font(TitleFont)
+				.Text(LOCTEXT("SuffixTitle", "Suffix"))
+			]
+
+			// Suffix name
+			+ SCanvas::Slot()
+			.Position(FVector2D(0.f, TitleOffsetY) + CheckboxOffset)
+			.Size(CheckboxSize)
+			[
+				SAssignNew(SuffixCheckBox, SCheckBox)
+				.IsChecked(this, &SAdvancedRenamerPanel::IsSuffixChecked)
+				.OnCheckStateChanged(this, &SAdvancedRenamerPanel::OnSuffixCheckBoxChanged)
+			]
+			+ SCanvas::Slot()
+			.Position(FVector2D(LabelStart, TitleOffsetY + LabelOffsetY))
+			.Size(LabelSize)
+			[
+				SNew(STextBlock)
+				.ColorAndOpacity(FStyleColors::AccentWhite.GetSpecifiedColor())
+				.Font(RegularFont)
+				.Text(LOCTEXT("Suffix", "Suffix"))
+			]
+			+ SCanvas::Slot()
+			.Position(FVector2D(EntryStart, TitleOffsetY))
+			.Size(EntrySize)
+			[
+				SAssignNew(SuffixTextBox, SEditableTextBox)
+				.BackgroundColor(FStyleColors::Background.GetSpecifiedColor())
+				.ForegroundColor(FStyleColors::AccentWhite.GetSpecifiedColor())
+				.Font(RegularFont)
+				.HintText(LOCTEXT("SuffixHint", "Suffix"))
+				.IsEnabled(false)
+				.OnTextChanged(this, &SAdvancedRenamerPanel::OnSuffixChanged)
+			]
+
+			// Remove Suffix
+			+ SCanvas::Slot()
+			.Position(FVector2D(LabelStart, TitleOffsetY + LineHeight) + RadioOffset)
+			.Size(CheckboxSize)
+			[
+				SAssignNew(SuffixRemoveCheckBox, SCheckBox)
+				.Style(&FAdvancedRenamerStyle::Get().GetWidgetStyle<FCheckBoxStyle>("AdvancedRenamer.Style.BlackRadioButton"))
+				.IsChecked(this, &SAdvancedRenamerPanel::IsSuffixRemoveChecked)
+				.OnCheckStateChanged(this, &SAdvancedRenamerPanel::OnSuffixRemoveCheckBoxChanged)
+				.IsEnabled(false)
+			]
+			+ SCanvas::Slot()
+			.Position(FVector2D(LabelStart * 2.f, TitleOffsetY + LabelOffsetY + LineHeight))
+			.Size(LabelSize)
+			[
+				SNew(STextBlock)
+				.ColorAndOpacity(FStyleColors::AccentWhite.GetSpecifiedColor())
+				.Font(RegularFont)
+				.Text(LOCTEXT("SuffixRemove", "Remove Suffix"))
+			]
+			+ SCanvas::Slot()
+			.Position(FVector2D(EntryStart, TitleOffsetY + LineHeight))
+			.Size(SeparatorSize)
+			[
+				SAssignNew(SuffixSeparatorTextBox, SEditableTextBox)
+				.BackgroundColor(FStyleColors::Background.GetSpecifiedColor())
+				.ForegroundColor(FStyleColors::AccentWhite.GetSpecifiedColor())
+				.Font(RegularFont)
+				.IsEnabled(false)
+				.Text(LOCTEXT("Underscore", "_"))
+				.OnVerifyTextChanged(this, &SAdvancedRenamerPanel::OnSuffixSeparatorVerifyTextChanged)
+				.OnTextChanged(this, &SAdvancedRenamerPanel::OnSuffixSeparatorChanged)
+			]
+			+ SCanvas::Slot()
+			.Position(FVector2D(EntryStart + SeparatorSize.X + 5.f, TitleOffsetY + LabelOffsetY + LineHeight))
+			.Size(LabelSize)
+			[
+				SNew(STextBlock)
+				.ColorAndOpacity(FStyleColors::AccentWhite.GetSpecifiedColor())
+				.Font(RegularFont)
+				.Text(LOCTEXT("Separator", "(separator)"))
+			]
+
+			// Remove first
+			+ SCanvas::Slot()
+			.Position(FVector2D(LabelStart, TitleOffsetY + LineHeight * 2.f) + RadioOffset)
+			.Size(CheckboxSize)
+			[
+				SAssignNew(SuffixRemoveCharactersCheckBox, SCheckBox)
+				.Style(&FAdvancedRenamerStyle::Get().GetWidgetStyle<FCheckBoxStyle>("AdvancedRenamer.Style.BlackRadioButton"))
+				.IsChecked(this, &SAdvancedRenamerPanel::IsSuffixRemoveCharactersChecked)
+				.OnCheckStateChanged(this, &SAdvancedRenamerPanel::OnSuffixRemoveCharactersCheckBoxChanged)
+				.IsEnabled(false)
+			]
+			+ SCanvas::Slot()
+			.Position(FVector2D(LabelStart * 2.f, TitleOffsetY + LabelOffsetY + LineHeight * 2.f))
+			.Size(LabelSize)
+			[
+				SNew(STextBlock)
+				.ColorAndOpacity(FStyleColors::AccentWhite.GetSpecifiedColor())
+				.Font(RegularFont)
+				.Text(LOCTEXT("SuffixRemoveChars", "Remove Last"))
+			]
+			+ SCanvas::Slot()
+			.Position(FVector2D(EntryStart, TitleOffsetY + LineHeight * 2.f))
+			.Size(SpinSize1)
+			[
+
+				SAssignNew(SuffixRemoveCharactersSpinBox, SSpinBox<uint8>)
+				.Style(&FAppStyle::Get(), "Menu.SpinBox")
+				.Font(RegularFont)
+				.MinValue(1)
+				.MaxValue(9)
+				.Value(1)
+				.IsEnabled(false)
+				.OnValueChanged(this, &SAdvancedRenamerPanel::OnSuffixRemoveCharactersChanged)
+			]
+			+ SCanvas::Slot()
+			.Position(FVector2D(EntryStart + SeparatorSize.X + 5.f, TitleOffsetY + LabelOffsetY + LineHeight * 2.f))
+			.Size(LabelSize)
+			[
+				SNew(STextBlock)
+				.ColorAndOpacity(FStyleColors::AccentWhite.GetSpecifiedColor())
+				.Font(RegularFont)
+				.Text(LOCTEXT("Characters", "character(s)"))
+			]
+
+			// Remove Number
+			+ SCanvas::Slot()
+			.Position(FVector2D(LabelStart, TitleOffsetY + LineHeight * 3.f) + CheckboxOffset)
+			.Size(CheckboxSize)
+			[
+				SAssignNew(SuffixRemoveNumberCheckBox, SCheckBox)
+				.IsChecked(this, &SAdvancedRenamerPanel::IsSuffixRemoveNumberChecked)
+				.OnCheckStateChanged(this, &SAdvancedRenamerPanel::OnSuffixRemoveNumberCheckBoxChanged)
+				.IsEnabled(false)
+			]
+			+ SCanvas::Slot()
+			.Position(FVector2D(LabelStart * 2.f, TitleOffsetY + LabelOffsetY + LineHeight * 3.f))
+			.Size(LabelSize)
+			[
+				SNew(STextBlock)
+				.ColorAndOpacity(FStyleColors::AccentWhite.GetSpecifiedColor())
+				.Font(RegularFont)
+				.Text(LOCTEXT("SuffixRemoveNumber", "Remove Number"))
+			]
+
+			// Number
+			+ SCanvas::Slot()
+			.Position(FVector2D(LabelStart, TitleOffsetY + LineHeight * 4.f) + CheckboxOffset)
+			.Size(CheckboxSize)
+			[
+				SAssignNew(SuffixNumberCheckBox, SCheckBox)
+				.IsChecked(this, &SAdvancedRenamerPanel::IsSuffixNumberChecked)
+				.OnCheckStateChanged(this, &SAdvancedRenamerPanel::OnSuffixNumberCheckBoxChanged)
+				.IsEnabled(false)
+			]
+			+ SCanvas::Slot()
+			.Position(FVector2D(LabelStart * 2.f, TitleOffsetY + LabelOffsetY + LineHeight * 4.f))
+			.Size(LabelSize)
+			[
+				SNew(STextBlock)
+				.ColorAndOpacity(FStyleColors::AccentWhite.GetSpecifiedColor())
+				.Font(RegularFont)
+				.Text(LOCTEXT("SuffixNumber", "Number"))
+			]
+			+ SCanvas::Slot()
+			.Position(FVector2D(EntryStart , TitleOffsetY + LabelOffsetY + LineHeight * 4.f))
+			.Size(LabelSize)
+			[
+				SNew(STextBlock)
+				.ColorAndOpacity(FStyleColors::AccentWhite.GetSpecifiedColor())
+				.Font(RegularFont)
+				.Text(LOCTEXT("Start", "Start"))
+			]
+			+ SCanvas::Slot()
+			.Position(FVector2D(EntryStart + 33.f, TitleOffsetY + LineHeight * 4.f))
+			.Size(SpinSize3)
+			[
+
+				SAssignNew(SuffixNumberStartSpinBox, SSpinBox<int32>)
+				.Style(&FAppStyle::Get(), "Menu.SpinBox")
+				.Font(RegularFont)
+				.MinValue(0)
+				.MaxValue(999)
+				.Value(1)
+				.IsEnabled(false)
+				.OnValueChanged(this, &SAdvancedRenamerPanel::OnSuffixNumberStartChanged)
+			]
+			+ SCanvas::Slot()
+			.Position(FVector2D(EntryStart + 33.f + SpinSize3.X + 5.f, TitleOffsetY + LabelOffsetY + LineHeight * 4.f))
+			.Size(LabelSize)
+			[
+				SNew(STextBlock)
+				.ColorAndOpacity(FStyleColors::AccentWhite.GetSpecifiedColor())
+				.Font(RegularFont)
+				.Text(LOCTEXT("Step", "Step"))
+			]
+			+ SCanvas::Slot()
+			.Position(FVector2D(EntryStart + 33.f + SpinSize3.X + 36.f, TitleOffsetY + LineHeight * 4.f))
+			.Size(SpinSize2)
+			[
+
+				SAssignNew(SuffixNumberStepSpinBox, SSpinBox<int32>)
+				.Style(&FAppStyle::Get(), "Menu.SpinBox")
+				.Font(RegularFont)
+				.MinValue(1)
+				.MaxValue(99)
+				.Value(1)
+				.IsEnabled(false)
+				.OnValueChanged(this, &SAdvancedRenamerPanel::OnSuffixNumberStepChanged)
+			]
+		];
+	// @formatter:on
+}
+
+TSharedRef<SWidget> SAdvancedRenamerPanel::CreateSearchAndReplace()
+{
+	using namespace UE::AdvancedRenamer::Private;
+
+	// @formatter:off
+	return SNew(SBorder)
+		.BorderImage(FAppStyle::Get().GetBrush("DetailsView.CategoryTop"))
+		.Content()
+		[
+			SNew(SCanvas)
+
+			// Title
+			+ SCanvas::Slot()
+			.Position(FVector2D(0.f, LabelOffsetY))
+			.Size(TitleSize)
+			[
+				SNew(STextBlock)
+				.ColorAndOpacity(FStyleColors::AccentBlue.GetSpecifiedColor())
+				.Font(TitleFont)
+				.Text(LOCTEXT("SearchReplaceTitle", "Search and Replace"))
+			]
+
+			+ SCanvas::Slot()
+			.Position(FVector2D(0.f, TitleOffsetY) + RadioOffset)
+			.Size(CheckboxSize)
+			[
+				SAssignNew(SearchReplacePlainTextCheckbox, SCheckBox)
+				.Style(&FAdvancedRenamerStyle::Get().GetWidgetStyle<FCheckBoxStyle>("AdvancedRenamer.Style.BlackRadioButton"))
+				.IsChecked(this, &SAdvancedRenamerPanel::IsSearchReplacePlainTextChecked)
+				.OnCheckStateChanged(this, &SAdvancedRenamerPanel::OnSearchReplacePlainTextCheckBoxChanged)
+			]
+
+			+ SCanvas::Slot()
+			.Position(FVector2D(0.f + LabelStart, TitleOffsetY + LabelOffsetY))
+			.Size(LabelSize)
+			[
+				SNew(STextBlock)
+				.ColorAndOpacity(FStyleColors::AccentWhite.GetSpecifiedColor())
+				.Font(RegularFont)
+				.Text(LOCTEXT("PlainText", "Plain Text"))
+			]
+
+			+ SCanvas::Slot()
+			.Position(FVector2D(92.f, TitleOffsetY) + RadioOffset)
+			.Size(CheckboxSize)
+			[
+				SAssignNew(SearchReplaceRegexCheckbox, SCheckBox)
+				.Style(&FAdvancedRenamerStyle::Get().GetWidgetStyle<FCheckBoxStyle>("AdvancedRenamer.Style.BlackRadioButton"))
+				.IsChecked(this, &SAdvancedRenamerPanel::IsSearchReplaceRegexChecked)
+				.OnCheckStateChanged(this, &SAdvancedRenamerPanel::OnSearchReplaceRegexCheckBoxChanged)
+			]
+
+			+ SCanvas::Slot()
+			.Position(FVector2D(92.f + LabelStart, TitleOffsetY + LabelOffsetY))
+			.Size(LabelSize)
+			[
+				SNew(STextBlock)
+				.ColorAndOpacity(FStyleColors::AccentWhite.GetSpecifiedColor())
+				.Font(RegularFont)
+				.Text(LOCTEXT("Regex", "Regex"))
+			]
+
+			+ SCanvas::Slot()
+			.Position(FVector2D(200.f, TitleOffsetY) + CheckboxOffset)
+			.Size(CheckboxSize)
+			[
+				SAssignNew(SearchReplaceIgnoreCaseCheckBox, SCheckBox)
+				.IsChecked(this, &SAdvancedRenamerPanel::IsSearchReplaceIgnoreCaseChecked)
+				.OnCheckStateChanged(this, &SAdvancedRenamerPanel::OnSearchReplaceIgnoreCaseCheckBoxChanged)
+			]
+
+			+ SCanvas::Slot()
+			.Position(FVector2D(200.f + LabelStart, TitleOffsetY + LabelOffsetY))
+			.Size(LabelSize)
+			[
+				SNew(STextBlock)
+				.ColorAndOpacity(FStyleColors::AccentWhite.GetSpecifiedColor())
+				.Font(RegularFont)
+				.Text(LOCTEXT("IgnoreCase", "Ignore Case"))
+			]
+
+			+ SCanvas::Slot()
+			.Position(FVector2D(0.f, TitleOffsetY + LineHeight))
+			.Size(FVector2D(LeftBlockWidth - 4.f, 54.f))
+			[
+				SAssignNew(SearchReplaceSearchTextBox, SMultiLineEditableTextBox)
+				.BackgroundColor(FStyleColors::Background.GetSpecifiedColor())
+				.ForegroundColor(FStyleColors::AccentWhite.GetSpecifiedColor())
+				.Font(RegularFont)
+				.AllowMultiLine(false)
+				.AutoWrapText(true)
+				.HintText(LOCTEXT("RegeSearchHint", "Search"))
+				.WrappingPolicy(ETextWrappingPolicy::AllowPerCharacterWrapping)
+				.IsEnabled(false)
+				.OnTextChanged(this, &SAdvancedRenamerPanel::OnSearchReplaceSearchTextChanged)
+			]
+
+			+ SCanvas::Slot()
+			.Position(FVector2D(0.f, TitleOffsetY + LineHeight + 59.f))
+			.Size(FVector2D(LeftBlockWidth - 4.f, 54.f))
+			[
+				SAssignNew(SearchReplaceReplaceTextBox, SMultiLineEditableTextBox)
+				.BackgroundColor(FStyleColors::Background.GetSpecifiedColor())
+				.ForegroundColor(FStyleColors::AccentWhite.GetSpecifiedColor())
+				.Font(RegularFont)
+				.AllowMultiLine(false)
+				.AutoWrapText(true)
+				.HintText(LOCTEXT("RegeReplaceHint", "Replace"))
+				.WrappingPolicy(ETextWrappingPolicy::AllowPerCharacterWrapping)
+				.IsEnabled(false)
+				.OnTextChanged(this, &SAdvancedRenamerPanel::OnSearchReplaceReplaceTextChanged)
+			]
+		];
+	// @formatter:on
+}
+
+void SAdvancedRenamerPanel::CreateRightPane(TSharedRef<SCanvas> Canvas)
+{
+	using namespace UE::AdvancedRenamer::Private;
+
+	// @formatter:off
+	Canvas->AddSlot()
+		.Position(FVector2D(RightBlockOffsetX, 5.f))
+		.Size(RightBlockSize)
+		[
+			CreateRenamePreview()
+		];
+	// @formatter:on
+}
+
+TSharedRef<SWidget> SAdvancedRenamerPanel::CreateRenamePreview()
+{
+	using namespace UE::AdvancedRenamer::Private;
+
+	// @formatter:off
+	return SNew(SBorder)
+		.BorderImage(FAppStyle::Get().GetBrush("DetailsView.CategoryTop"))
+		.Content()
+		[
+			SNew(SCanvas)
+				
+			// Title
+			+ SCanvas::Slot()
+			.Position(FVector2D(0.f, LabelOffsetY))
+			.Size(TitleSize)
+			[
+				SNew(STextBlock)
+				.ColorAndOpacity(FStyleColors::AccentBlue.GetSpecifiedColor())
+				.Font(TitleFont)
+				.Text(LOCTEXT("ObjectRenamePreviewTitle", "Rename Preview"))
+			]
+
+			// List View
+			+ SCanvas::Slot()
+			.Position(FVector2D(0.f, TitleOffsetY))
+			.Size(ListViewSize)
+			[
+				SNew(SScrollBox)
+				.Orientation(EOrientation::Orient_Horizontal)
+				+ SScrollBox::Slot()
+				[
+					SAssignNew(RenamePreviewListBox, SBox)
+					.WidthOverride(ListViewSize.Y)
+					.HeightOverride(ListViewSize.X)
+					.HAlign(EHorizontalAlignment::HAlign_Fill)
+					.VAlign(EVerticalAlignment::VAlign_Fill)
+					.Content()
+					[
+						SAssignNew(RenamePreviewList, SListView<FObjectRenamePreviewListItemPtr>)
+						.ItemHeight(ListLineHeight)
+						.ListItemsSource(&ListData)
+						.OnGenerateRow(this, &SAdvancedRenamerPanel::OnGenerateRowForList)
+						.HeaderRow(
+							SAssignNew(RenamePreviewListHeaderRow, SHeaderRow)
+							+ SHeaderRow::Column(FAdvancedRenamerPreviewListItem::OriginalNameColumnName)
+							.DefaultLabel(LOCTEXT("From", "From"))
+							.FillWidth(ListViewSize.X / 2.f)
+							+ SHeaderRow::Column(FAdvancedRenamerPreviewListItem::NewNameColumnName)
+							.DefaultLabel(LOCTEXT("To", "To"))
+							.FillWidth(ListViewSize.X / 2.f)
+						)
+						.OnKeyDownHandler(this, &SAdvancedRenamerPanel::OnListViewKeyDown)
+						.OnContextMenuOpening(this, &SAdvancedRenamerPanel::GenerateListViewContextMenu)
+					]
+				]
+			]
+
+			// Apply Button
+			+ SCanvas::Slot()
+			.Position(FVector2D(0.f, TitleOffsetY + ListViewSize.Y + 5.f))
+			.Size(FVector2D(ListViewSize.X, ApplyButtonHeight))
+			[
+				SAssignNew(ApplyButton, SButton)
+				.ButtonStyle(FAdvancedRenamerStyle::Get(), "AdvancedRenamer.Style.DarkButton")
+				.ToolTipText(LOCTEXT("ApplyRenameToolTip", "Renames all actors."))
+				.HAlign(EHorizontalAlignment::HAlign_Center)
+				.VAlign(EVerticalAlignment::VAlign_Center)
+				.IsEnabled(false)
+				.OnClicked(this, &SAdvancedRenamerPanel::OnApplyButtonClicked)
+				[
+					SNew(STextBlock)
+					.ColorAndOpacity(FStyleColors::AccentWhite.GetSpecifiedColor())
+					.Font(RegularFont)
+					.Text(LOCTEXT("ApplyRename", "Apply"))
+				]
+			]
+		];
+	// @formatter:on
+}
+
+bool SAdvancedRenamerPanel::RenameObjects()
+{
+	if (!bValidNames)
+	{
+		return false;
+	}
+
+	FScopedTransaction Transaction(LOCTEXT("AdvancedObjectRename", "Advanced Object Rename"));
+	int32 ItemsRenamed = 0;
+	int32 ItemCount = Num();
+	
+	for (int32 Index = 0; Index < ItemCount; ++Index)
+	{
+		if (!ListData[Index].IsValid())
+		{
+			continue;
+		}
+
+		if (!IsValidIndex(Index))
+		{
+			continue;
+		}
+
+		if (ListData[Index]->NewName.Len() == 0)
+		{
+			continue;
+		}
+
+		if (ExecuteRename(Index, ListData[Index]->NewName))
+		{
+			++ItemsRenamed;
+		}
+	}
+
+	//Cancel Transaction if no Items could successfully Rename
+	if (ItemsRenamed == 0)
+	{
+		Transaction.Cancel();
+		//TODO: Should we keep Window Open as no renames occurred. Might be something the User did not expect
+		//return false;
+	}
+	
+	return true;
+}
+
+bool SAdvancedRenamerPanel::CloseWindow()
+{
+	TSharedPtr<SWindow> CurrentWindow = FSlateApplication::Get().FindWidgetWindow(SharedThis(this));
+
+	if (CurrentWindow.IsValid())
+	{
+		CurrentWindow->RequestDestroyWindow();
+		return true;
+	}
+
+	return false;
+}
+
+FString SAdvancedRenamerPanel::CreateNewName(int32 Index) const
+{
+	if (!IsValidIndex(Index))
+	{
+		return "";
+	}
+
+	FString DisplayName = GetOriginalName(Index);
+
+	if (DisplayName.Len() == 0)
+	{
+		return DisplayName;
+	}
+
+	return ApplyRename(DisplayName, Index);
+}
+
+FString SAdvancedRenamerPanel::ApplyRename(const FString& OriginalName, int32 Index) const
+{
+	FString NewName = OriginalName;
+
+	if (bBaseName)
+	{
+		NewName = ApplyBaseName(NewName);
+	}
+
+	if (bPrefix)
+	{
+		NewName = ApplyPrefix(NewName);
+	}
+
+	if (bSuffix)
+	{
+		NewName = ApplySuffix(NewName, Index);
+	}
+
+	if (bSearchReplacePlainText)
+	{
+		NewName = ApplySearchPlainText(NewName);
+	}
+	else if (bSearchReplaceRegex)
+	{
+		NewName = ApplySearchReplaceRegex(NewName);
+	}
+
+	return NewName;
+}
+
+FString SAdvancedRenamerPanel::ApplyBaseName(const FString& OriginalName) const
+{
+	if (!BaseNameTextBox.IsValid())
+	{
+		return OriginalName;
+	}
+
+	FString BaseName = BaseNameTextBox->GetText().ToString();
+	
+	if (BaseName.Len() == 0)
+	{
+		return OriginalName;
+	}
+
+	return BaseName;
+}
+
+FString SAdvancedRenamerPanel::ApplyPrefix(const FString& OriginalName) const
+{
+	FString Output = OriginalName;
+
+	if (bPrefixRemove)
+	{
+		if (PrefixSeparatorTextBox.IsValid())
+		{
+			FString Separator = PrefixSeparatorTextBox->GetText().ToString();
+
+			if (Separator.Len() != 0)
+			{
+				int32 PrefixStart = Output.Find(Separator, ESearchCase::IgnoreCase);
+
+				if (PrefixStart >= 0)
+				{
+					Output = Output.Mid(PrefixStart + Separator.Len());
+				}
+			}
+		}
+	}
+	else if (bPrefixRemoveCharacters)
+	{
+		if (PrefixRemoveCharactersSpinBox.IsValid())
+		{
+			Output = Output.RightChop(PrefixRemoveCharactersSpinBox->GetValue());
+		}
+	}
+
+	if (PrefixTextBox.IsValid())
+	{
+		FString Prefix = PrefixTextBox->GetText().ToString();
+
+		if (Prefix.Len() > 0)
+		{
+			Output = Prefix + Output;
+		}
+	}
+
+	return Output;
+}
+
+FString SAdvancedRenamerPanel::ApplySuffix(const FString& OriginalName, int32 Index) const
+{
+	static const TCHAR FirstDigit = '0';
+	static const TCHAR LastDigit = '9';
+
+	FString Output = OriginalName;
+
+	if (bSuffixRemove)
+	{
+		if (SuffixSeparatorTextBox.IsValid())
+		{
+			FString Separator = SuffixSeparatorTextBox->GetText().ToString();
+
+			if (Separator.Len() != 0)
+			{
+				int32 SuffixStart = Output.Find(Separator, ESearchCase::IgnoreCase, ESearchDir::FromEnd);
+
+				if (SuffixStart >= 0)
+				{
+					Output = Output.Mid(0, SuffixStart);
+				}
+			}
+		}
+	}
+	else if (bSuffixRemoveCharacters)
+	{
+		if (SuffixRemoveCharactersSpinBox.IsValid())
+		{
+			Output = Output.LeftChop(SuffixRemoveCharactersSpinBox->GetValue());
+		}
+	}
+
+	if (bSuffixRemoveNumber)
+	{
+		int32 LastDigitIndex = Output.Len() - 1;
+
+		while(LastDigitIndex >= 0)
+		{
+			if (Output[LastDigitIndex] < FirstDigit || Output[LastDigitIndex] > LastDigit)
+			{
+				break;
+			}
+
+			--LastDigitIndex;
+		}
+
+		Output = Output.Mid(0, LastDigitIndex + 1);
+	}
+
+	if (SuffixTextBox.IsValid())
+	{
+		FString Suffix = SuffixTextBox->GetText().ToString();
+
+		if (Suffix.Len() > 0)
+		{
+			Output = Output + Suffix;
+		}
+	}
+
+	if (bSuffixNumber)
+	{
+		if (SuffixNumberStartSpinBox.IsValid() && SuffixNumberStepSpinBox.IsValid())
+		{
+			int32 Start = SuffixNumberStartSpinBox->GetValue();
+			int32 Step = SuffixNumberStepSpinBox->GetValue();
+			int32 Current = Start + (Index * Step);
+
+			Output += FString::FromInt(Current);
+		}
+	}
+
+	return Output;
+}
+
+FString SAdvancedRenamerPanel::ApplySearchPlainText(const FString& OriginalName) const
+{
+	if (!SearchReplaceSearchTextBox.IsValid() || !SearchReplaceReplaceTextBox.IsValid())
+	{
+		return OriginalName;
+	}
+
+	FString Search = SearchReplaceSearchTextBox->GetPlainText().ToString();
+
+	if (Search.Len() == 0)
+	{
+		return OriginalName;
+	}
+
+	FString Replace = SearchReplaceReplaceTextBox->GetPlainText().ToString();
+
+	if (Replace.Len() == 0)
+	{
+		UE_LOG(LogARP, Warning, TEXT("Regex: Empty replacement string."));
+	}
+
+	return OriginalName.Replace(*Search, *Replace, bSearchReplaceIgnoreCase ? ESearchCase::IgnoreCase : ESearchCase::CaseSensitive);
+}
+
+FString SAdvancedRenamerPanel::ApplySearchReplaceRegex(const FString& OriginalName) const
+{
+	if (!SearchReplaceSearchTextBox.IsValid() || !SearchReplaceReplaceTextBox.IsValid())
+	{
+		return OriginalName;
+	}
+
+	FString Pattern = SearchReplaceSearchTextBox->GetPlainText().ToString();
+
+	if (Pattern.Len() == 0)
+	{
+		return OriginalName;
+	}
+	
+	FRegexPattern RegexPattern = FRegexPattern(
+		Pattern, 
+		bSearchReplaceIgnoreCase ? ERegexPatternFlags::CaseInsensitive : ERegexPatternFlags::None
+	);
+
+	FString ReplaceString = SearchReplaceReplaceTextBox->GetPlainText().ToString();
+
+	if (ReplaceString.Len() == 0)
+	{
+		UE_LOG(LogARP, Warning, TEXT("Regex: Empty replacement string."));
+	}
+
+	return RegexReplace(OriginalName, RegexPattern, ReplaceString);
+}
+
+FString SAdvancedRenamerPanel::RegexReplace(const FString& OriginalString, const FRegexPattern Pattern, const FString& ReplaceString) const
+{
+	static const FString EscapeString = "\\";
+	static const TCHAR EscapeChar = EscapeString[0];
+	static const FString GroupString = "$";
+	static const TCHAR GroupChar = GroupString[0];
+	static const TCHAR FirstDigit = '0';
+	static const TCHAR LastDigit = '9';
+	static const TCHAR NullChar = 0;
+
+	FRegexMatcher Matcher(Pattern, OriginalString);
+
+	FString Output = "";
+	int32 StartCharIdx = 0;
+	bool bReadingGroupName = false;
+	int32 GroupIndex = INDEX_NONE;
+	
+	while (Matcher.FindNext())
+	{
+		// Add on part of string after start/previous match
+		if (StartCharIdx != Matcher.GetMatchBeginning())
+		{
+			Output += OriginalString.Mid(StartCharIdx, Matcher.GetMatchBeginning() - StartCharIdx);
+		}
+
+		bool bEscaped = false;
+
+		for (int32 CharIdx = 0; CharIdx <= ReplaceString.Len(); ++CharIdx)
+		{
+			const TCHAR& Char = CharIdx < ReplaceString.Len() ? ReplaceString[CharIdx] : NullChar;
+
+			if (bReadingGroupName)
+			{
+				// Build group index
+				if (Char >= FirstDigit && Char <= LastDigit)
+				{
+					if (GroupIndex == INDEX_NONE)
+					{
+						GroupIndex = 0;
+					}
+					else if (GroupIndex > 0)
+					{
+						GroupIndex *= 10;
+					}
+
+					int32 NextDigit = static_cast<int32>(Char - FirstDigit);					
+					GroupIndex += NextDigit;
+					continue;
+				}
+				// We've read a group index, add it to the output string
+				else if (GroupIndex > 0)
+				{
+					if (Matcher.GetCaptureGroupBeginning(GroupIndex) == INDEX_NONE)
+					{
+						UE_LOG(LogARP, Error, TEXT("Regex: Capture group does not exist %d."), GroupIndex);
+					}
+
+					Output += Matcher.GetCaptureGroup(GroupIndex);
+				}
+				// $0 matches the entire matched string
+				else if (GroupIndex == 0)
+				{
+					Output += OriginalString.Mid(Matcher.GetMatchBeginning(), Matcher.GetMatchEnding() - Matcher.GetMatchBeginning());
+				}
+				// An unescaped $
+				else
+				{
+					UE_LOG(LogARP, Error, TEXT("Regex: Unescaped %s."), *GroupString);
+
+					Output += GroupString;
+				}
+
+				bReadingGroupName = false;
+				// Continue regular parsing of this character.
+			}
+
+			// Check for special chars
+			if (!bEscaped)
+			{
+				if (Char == EscapeChar)
+				{
+					bEscaped = true;
+					continue;
+				}
+
+				if (Char == GroupChar)
+				{
+					bReadingGroupName = true;
+					GroupIndex = INDEX_NONE;
+					continue;
+				}
+			}
+			else
+			{
+				// If the last char is a \ assume that it's not an escape char
+				if (Char == NullChar)
+				{
+					UE_LOG(LogARP, Error, TEXT("Regex: Unescaped %s."), *EscapeString);
+
+					Output += EscapeChar;
+				}
+			}
+
+			if (Char != NullChar)
+			{
+				Output += Char;
+			}
+
+			bEscaped = false;
+		}
+
+		StartCharIdx = Matcher.GetMatchEnding();
+	}
+
+	// Add on the end of the string after the last match
+	if (StartCharIdx < OriginalString.Len())
+	{
+		Output += OriginalString.Mid(StartCharIdx);
+	}
+
+	return Output;
+}
+
+void SAdvancedRenamerPanel::UpdateEnables()
+{
+	if (BaseNameTextBox.IsValid())
+	{
+		BaseNameTextBox->SetEnabled(bBaseName);
+	}
+
+	if (PrefixTextBox.IsValid())
+	{
+		PrefixTextBox->SetEnabled(bPrefix);
+	}
+
+	if (PrefixRemoveCheckBox.IsValid())
+	{
+		PrefixRemoveCheckBox->SetEnabled(bPrefix);
+	}
+
+	if (PrefixSeparatorTextBox.IsValid())
+	{
+		PrefixSeparatorTextBox->SetEnabled(bPrefix && bPrefixRemove);
+	}
+
+	if (PrefixRemoveCharactersCheckBox.IsValid())
+	{
+		PrefixRemoveCharactersCheckBox->SetEnabled(bPrefix);
+	}
+
+	if (PrefixRemoveCharactersSpinBox.IsValid())
+	{
+		PrefixRemoveCharactersSpinBox->SetEnabled(bPrefix && bPrefixRemoveCharacters);
+	}
+
+	if (SuffixTextBox.IsValid())
+	{
+		SuffixTextBox->SetEnabled(bSuffix);
+	}
+
+	if (SuffixRemoveCheckBox.IsValid())
+	{
+		SuffixRemoveCheckBox->SetEnabled(bSuffix);
+	}
+
+	if (SuffixSeparatorTextBox.IsValid())
+	{
+		SuffixSeparatorTextBox->SetEnabled(bSuffix && bSuffixRemove);
+	}
+
+	if (SuffixRemoveCharactersCheckBox.IsValid())
+	{
+		SuffixRemoveCharactersCheckBox->SetEnabled(bSuffix);
+	}
+
+	if (SuffixRemoveCharactersSpinBox.IsValid())
+	{
+		SuffixRemoveCharactersSpinBox->SetEnabled(bSuffix && bSuffixRemoveCharacters);
+	}
+
+	if (SuffixRemoveNumberCheckBox.IsValid())
+	{
+		SuffixRemoveNumberCheckBox->SetEnabled(bSuffix && !bSuffixNumber);
+	}
+
+	if (SuffixNumberCheckBox.IsValid())
+	{
+		SuffixNumberCheckBox->SetEnabled(bSuffix);
+	}
+
+	if (SuffixNumberStartSpinBox.IsValid())
+	{
+		SuffixNumberStartSpinBox->SetEnabled(bSuffix && bSuffixNumber);
+	}
+
+	if (SuffixNumberStepSpinBox.IsValid())
+	{
+		SuffixNumberStepSpinBox->SetEnabled(bSuffix && bSuffixNumber);
+	}
+
+	if (SearchReplaceIgnoreCaseCheckBox.IsValid())
+	{
+		SearchReplaceIgnoreCaseCheckBox->SetEnabled(bSearchReplacePlainText || bSearchReplaceRegex);
+	}
+
+	if (SearchReplaceSearchTextBox.IsValid())
+	{
+		SearchReplaceSearchTextBox->SetEnabled(bSearchReplacePlainText || bSearchReplaceRegex);
+	}
+
+	if (SearchReplaceReplaceTextBox.IsValid())
+	{
+		SearchReplaceReplaceTextBox->SetEnabled(bSearchReplacePlainText || bSearchReplaceRegex);
+	}
+
+	if (ApplyButton.IsValid())
+	{
+		ApplyButton->SetEnabled(bValidNames);
+	}
+
+	RequestListViewRefresh();
+}
+
+void SAdvancedRenamerPanel::RequestListViewRefresh()
+{
+	bListNeedsUpdate = true;
+}
+
+void SAdvancedRenamerPanel::RefreshListView(const double InCurrentTime)
+{
+	bool bRemovedObjects = false;
+	bValidNames = false;
+
+	for (int32 Index = 0; Index < ListData.Num(); ++Index)
+	{
+		if (!ListData[Index].IsValid() || !IsValidIndex(Index))
+		{
+			RemoveIndex(Index);
+			--Index;
+			bRemovedObjects = true;
+			continue;
+		}
+
+		// Force recreation
+		ListData[Index]->NewName = CreateNewName(Index);
+
+		if (ListData[Index]->NewName.Len() == 0)
+		{
+			continue;
+		}
+
+		if (GetOriginalName(Index) == ListData[Index]->NewName)
+		{
+			continue;
+		}
+
+		bValidNames = true;
+	}
+
+	if (bRemovedObjects)
+	{
+		RenamePreviewList->RequestListRefresh();
+	}
+
+	UpdateEnables();
+
+	MinDesiredOriginalNameWidth = 0.f;
+	MinDesiredNewNameWidth = 0.f;
+	RenamePreviewList->RebuildList();
+
+	ListLastUpdateTime = InCurrentTime;
+	bListNeedsUpdate = false;
+}
+
+void SAdvancedRenamerPanel::UpdateRequiredListWidth()
+{
+	if (!RenamePreviewListBox.IsValid() || !RenamePreviewListHeaderRow.IsValid())
+	{
+		return;
+	}
+
+	if (MinDesiredNewNameWidth == 0.f)
+	{
+		return;
+	}
+
+	static const float ListViewWidth = 277.f;
+	float ActualWidth = FMath::Max(ListViewWidth, MinDesiredOriginalNameWidth + MinDesiredNewNameWidth + 20.f);
+
+	RenamePreviewListBox->SetWidthOverride(ActualWidth);
+	
+	RenamePreviewListHeaderRow->SetColumnWidth(
+		FAdvancedRenamerPreviewListItem::OriginalNameColumnName, 
+		(MinDesiredOriginalNameWidth + 10.f) / (MinDesiredOriginalNameWidth + MinDesiredNewNameWidth + 20.f) * ActualWidth
+	);
+
+	RenamePreviewListHeaderRow->SetColumnWidth(
+		FAdvancedRenamerPreviewListItem::NewNameColumnName, 
+		(MinDesiredNewNameWidth + 10.f) / (MinDesiredOriginalNameWidth + MinDesiredNewNameWidth + 20.f) * ActualWidth
+	);
+
+	RenamePreviewListHeaderRow->RefreshColumns();
+}
+
+void SAdvancedRenamerPanel::RemoveSelectedObjects()
+{
+	if (!RenamePreviewList.IsValid() || RenamePreviewList->GetNumItemsSelected() == 0)
+	{
+		return;
+	}
+
+	TArray<FObjectRenamePreviewListItemPtr> SelectedItems = RenamePreviewList->GetSelectedItems();
+
+	if (SelectedItems.Num() == 0)
+	{
+		return;
+	}
+
+	bool bMadeChange = false;
+
+	for (int32 Index = 0; Index < ListData.Num(); ++Index)
+	{
+		if (!ListData[Index].IsValid())
+		{
+			RemoveIndex(Index);
+			bMadeChange = true;
+			--Index;
+			continue;
+		}
+
+		int32 MatchIdx = INDEX_NONE;
+
+		for (int32 SelectedIdx = 0; SelectedIdx < SelectedItems.Num(); ++SelectedIdx)
+		{
+			if (!SelectedItems[SelectedIdx].IsValid())
+			{
+				continue;
+			}
+
+			if (SelectedItems[SelectedIdx]->Hash != ListData[Index]->Hash)
+			{
+				continue;
+			}
+
+			MatchIdx = SelectedIdx;
+			break;
+		}
+
+		if (MatchIdx == INDEX_NONE)
+		{
+			continue;
+		}
+
+		RemoveIndex(Index);
+		--Index;
+		SelectedItems.RemoveAt(MatchIdx);
+		bMadeChange = true;
+
+		if (SelectedItems.Num() == 0)
+		{
+			break;
+		}
+	}
+
+	if (bMadeChange)
+	{
+		if (ListData.Num() == 0)
+		{
+			if (CloseWindow())
+			{
+				return;
+			}
+		}
+
+		RenamePreviewList->RequestListRefresh();
+	}
+}
+
+void SAdvancedRenamerPanel::Tick(const FGeometry& AllottedGeometry, const double InCurrentTime, const float InDeltaTime)
+{
+	SCompoundWidget::Tick(AllottedGeometry, InCurrentTime, InDeltaTime);
+
+	if (ListLastUpdateTime == 0)
+	{
+		ListLastUpdateTime = InCurrentTime;
+	}
+	else if (bListNeedsUpdate && InCurrentTime >= (ListLastUpdateTime + MinUpdateFrequency))
+	{
+		RefreshListView(InCurrentTime);
+	}
+}
+
+void SAdvancedRenamerPanel::OnBaseNameCheckBoxChanged(ECheckBoxState NewState)
+{
+	bBaseName = NewState == ECheckBoxState::Checked;
+
+	UpdateEnables();
+}
+
+void SAdvancedRenamerPanel::OnBaseNameChanged(const FText& NewText)
+{
+	RequestListViewRefresh();
+}
+
+void SAdvancedRenamerPanel::OnPrefixCheckBoxChanged(ECheckBoxState NewState)
+{
+	bPrefix = NewState == ECheckBoxState::Checked;
+
+	UpdateEnables();
+}
+
+void SAdvancedRenamerPanel::OnPrefixChanged(const FText& NewText)
+{
+	RequestListViewRefresh();
+}
+
+void SAdvancedRenamerPanel::OnPrefixRemoveCheckBoxChanged(ECheckBoxState NewState)
+{
+	bPrefixRemove = NewState == ECheckBoxState::Checked;
+
+	if (bPrefixRemove)
+	{
+		bPrefixRemoveCharacters = false;
+	}
+
+	UpdateEnables();
+}
+
+bool SAdvancedRenamerPanel::OnPrefixSeparatorVerifyTextChanged(const FText& InText, FText& OutErrorText) const
+{
+	if (InText.ToString().Len() > 1)
+	{
+		OutErrorText = LOCTEXT("SeparatorError", "Separators can only be a single character.");
+		return false;
+	}
+
+	return true;
+}
+
+void SAdvancedRenamerPanel::OnPrefixSeparatorChanged(const FText& NewText)
+{
+	RequestListViewRefresh();
+}
+
+void SAdvancedRenamerPanel::OnPrefixRemoveCharactersCheckBoxChanged(ECheckBoxState NewState)
+{
+	bPrefixRemoveCharacters = NewState == ECheckBoxState::Checked;
+
+	if (bPrefixRemoveCharacters)
+	{
+		bPrefixRemove = false;
+	}
+
+	UpdateEnables();
+}
+
+void SAdvancedRenamerPanel::OnPrefixRemoveCharactersChanged(uint8 NewValue)
+{
+	RequestListViewRefresh();
+}
+
+void SAdvancedRenamerPanel::OnSuffixCheckBoxChanged(ECheckBoxState NewState)
+{
+	bSuffix = NewState == ECheckBoxState::Checked;
+
+	UpdateEnables();
+}
+
+void SAdvancedRenamerPanel::OnSuffixChanged(const FText& NewText)
+{
+	RequestListViewRefresh();
+}
+
+void SAdvancedRenamerPanel::OnSuffixRemoveCheckBoxChanged(ECheckBoxState NewState)
+{
+	bSuffixRemove = NewState == ECheckBoxState::Checked;
+
+	if (bSuffixRemove)
+	{
+		bSuffixRemoveCharacters = false;
+	}
+
+	UpdateEnables();
+}
+
+bool SAdvancedRenamerPanel::OnSuffixSeparatorVerifyTextChanged(const FText& InText, FText& OutErrorText) const
+{
+	if (InText.ToString().Len() > 1)
+	{
+		OutErrorText = LOCTEXT("SeparatorError", "Separators can only be a single character.");
+		return false;
+	}
+
+	return true;
+}
+
+void SAdvancedRenamerPanel::OnSuffixSeparatorChanged(const FText& NewText)
+{
+	RequestListViewRefresh();
+}
+
+void SAdvancedRenamerPanel::OnSuffixRemoveCharactersCheckBoxChanged(ECheckBoxState NewState)
+{
+	bSuffixRemoveCharacters = NewState == ECheckBoxState::Checked;
+
+	if (bSuffixRemoveCharacters)
+	{
+		bSuffixRemove = false;
+	}
+
+	UpdateEnables();
+}
+
+void SAdvancedRenamerPanel::OnSuffixRemoveCharactersChanged(uint8 NewValue)
+{
+	RequestListViewRefresh();
+}
+
+void SAdvancedRenamerPanel::OnSuffixRemoveNumberCheckBoxChanged(ECheckBoxState NewState)
+{
+	bSuffixRemoveNumber = NewState == ECheckBoxState::Checked;
+
+	UpdateEnables();
+}
+
+void SAdvancedRenamerPanel::OnSuffixNumberCheckBoxChanged(ECheckBoxState NewState)
+{
+	bSuffixNumber = NewState == ECheckBoxState::Checked;
+
+	if (bSuffixNumber)
+	{
+		bSuffixRemoveNumber = true;
+	}
+
+	UpdateEnables();
+}
+
+void SAdvancedRenamerPanel::OnSuffixNumberStartChanged(int32 NewValue)
+{
+	RequestListViewRefresh();
+}
+
+void SAdvancedRenamerPanel::OnSuffixNumberStepChanged(int32 NewValue)
+{
+	RequestListViewRefresh();
+}
+
+void SAdvancedRenamerPanel::OnSearchReplacePlainTextCheckBoxChanged(ECheckBoxState NewState)
+{
+	bSearchReplacePlainText = NewState == ECheckBoxState::Checked;
+
+	if (bSearchReplacePlainText)
+	{
+		bSearchReplaceRegex = false;
+	}
+
+	UpdateEnables();
+}
+
+void SAdvancedRenamerPanel::OnSearchReplaceRegexCheckBoxChanged(ECheckBoxState NewState)
+{
+	bSearchReplaceRegex = NewState == ECheckBoxState::Checked;
+
+	if (bSearchReplaceRegex)
+	{
+		bSearchReplacePlainText = false;
+	}
+
+	UpdateEnables();
+}
+
+void SAdvancedRenamerPanel::OnSearchReplaceIgnoreCaseCheckBoxChanged(ECheckBoxState NewState)
+{
+	bSearchReplaceIgnoreCase = NewState == ECheckBoxState::Checked;
+
+	UpdateEnables();
+}
+
+void SAdvancedRenamerPanel::OnSearchReplaceSearchTextChanged(const FText& NewText)
+{
+	RequestListViewRefresh();
+}
+
+void SAdvancedRenamerPanel::OnSearchReplaceReplaceTextChanged(const FText& NewText)
+{
+	RequestListViewRefresh();
+}
+
+TSharedRef<ITableRow> SAdvancedRenamerPanel::OnGenerateRowForList(FObjectRenamePreviewListItemPtr Item, 
+	const TSharedRef<STableViewBase>& OwnerTable)
+{
+	return SNew(SAdvancedRenamerPreviewListRow, SharedThis(this), OwnerTable, Item);
+}
+
+FReply SAdvancedRenamerPanel::OnListViewKeyDown(const FGeometry& MyGeometry, const FKeyEvent& KeyEvent)
+{
+	if (CommandList.IsValid() && CommandList->ProcessCommandBindings(KeyEvent))
+	{
+		return FReply::Handled();
+	}
+
+	return FReply::Unhandled();
+}
+
+TSharedPtr<SWidget> SAdvancedRenamerPanel::GenerateListViewContextMenu()
+{
+	if (!RenamePreviewList.IsValid())
+	{
+		return nullptr;
+	}
+
+	if (RenamePreviewList->GetNumItemsSelected() == 0)
+	{
+		return nullptr;
+	}
+
+	FMenuBuilder MenuBuilder(true, CommandList.ToSharedRef());
+
+	MenuBuilder.BeginSection("Actions", LOCTEXT("Actions", "Actions"));
+	{
+		MenuBuilder.AddMenuEntry(FGenericCommands::Get().Delete, NAME_None, LOCTEXT("RemoveObject", "Remove Object"));
+	}
+	MenuBuilder.EndSection();
+
+	return MenuBuilder.MakeWidget();
+}
+
+FReply SAdvancedRenamerPanel::OnApplyButtonClicked()
+{
+	if (RenameObjects())
+	{
+		CloseWindow();
+	}
+
+	return FReply::Handled();
+}
+
+int32 SAdvancedRenamerPanel::Num() const
+{
+	if (ObjectProvider.IsValid())
+	{
+		return IAdvancedRenamerObjectProvider::Execute_BP_Num(ObjectProvider.Get());
+	}
+	else
+	{
+		return SharedProvider->Num();
+	}
+}
+
+bool SAdvancedRenamerPanel::IsValidIndex(int32 Index) const
+{
+	if (ObjectProvider.IsValid())
+	{
+		return IAdvancedRenamerObjectProvider::Execute_BP_IsValidIndex(ObjectProvider.Get(), Index);
+	}
+	else
+	{
+		return SharedProvider->IsValidIndex(Index);
+	}
+}
+
+uint32 SAdvancedRenamerPanel::GetHash(int32 Index) const
+{
+	if (ObjectProvider.IsValid())
+	{
+		return static_cast<uint32>(IAdvancedRenamerObjectProvider::Execute_BP_GetHash(ObjectProvider.Get(), Index));
+	}
+	else
+	{
+		return SharedProvider->GetHash(Index);
+	}
+}
+
+FString SAdvancedRenamerPanel::GetOriginalName(int32 Index) const
+{
+	if (ObjectProvider.IsValid())
+	{
+		return IAdvancedRenamerObjectProvider::Execute_BP_GetOriginalName(ObjectProvider.Get(), Index);
+	}
+	else
+	{
+		return SharedProvider->GetOriginalName(Index);
+	}
+}
+
+bool SAdvancedRenamerPanel::RemoveIndex(int32 Index)
+{
+	// Can fail during construction when indices that aren't renameable are removed from the provider before
+	// they are added to ListData.
+	if (ListData.IsValidIndex(Index))
+	{
+		ListData.RemoveAt(Index);
+	}
+
+	if (ObjectProvider.IsValid())
+	{
+		return IAdvancedRenamerObjectProvider::Execute_BP_RemoveIndex(ObjectProvider.Get(), Index);
+	}
+	else
+	{
+		return SharedProvider->RemoveIndex(Index);
+	}
+}
+
+bool SAdvancedRenamerPanel::CanRename(int32 Index) const
+{
+	if (ObjectProvider.IsValid())
+	{
+		return IAdvancedRenamerObjectProvider::Execute_BP_CanRename(ObjectProvider.Get(), Index);
+	}
+	else
+	{
+		return SharedProvider->CanRename(Index);
+	}
+}
+
+bool SAdvancedRenamerPanel::ExecuteRename(int32 Index, const FString& NewName)
+{
+	if (ObjectProvider.IsValid())
+	{
+		return IAdvancedRenamerObjectProvider::Execute_BP_ExecuteRename(ObjectProvider.Get(), Index, NewName);
+	}
+	else
+	{
+		return SharedProvider->ExecuteRename(Index, NewName);
+	}
+}
+
+#undef LOCTEXT_NAMESPACE
diff --git a/Engine/Plugins/Experimental/AdvancedRenamer/Source/AdvancedRenamer/Private/Slate/SAdvancedRenamerPanel.h b/Engine/Plugins/Experimental/AdvancedRenamer/Source/AdvancedRenamer/Private/Slate/SAdvancedRenamerPanel.h
new file mode 100644
--- /dev/null
+++ b/Engine/Plugins/Experimental/AdvancedRenamer/Source/AdvancedRenamer/Private/Slate/SAdvancedRenamerPanel.h
@@ -0,0 +1,252 @@
+// Copyright Epic Games, Inc. All Rights Reserved.
+
+#pragma once
+
+#include "Providers/AdvancedRenamerProviderInterface.h"
+#include "Styling/SlateTypes.h"
+#include "UObject/StrongObjectPtr.h"
+#include "Widgets/Views/SListView.h"
+#include "Widgets/Views/STableRow.h"
+#include "Widgets/Views/STableViewBase.h"
+
+class FBlueprintEditor;
+class FRegexPattern;
+class FUICommandList;
+class IAdvancedRenamerProvider;
+class IAvaOutlinerItem;
+class SAdvancedRenamerPanel;
+class SBox;
+class SButton;
+class SCanvas;
+class SCheckBox;
+class SEditableTextBox;
+class SHeaderRow;
+class SMultiLineEditableTextBox; 
+template<typename NumericType>
+class SSpinBox;
+class UObject;
+typedef TSharedPtr<IAvaOutlinerItem> FAvaOutlinerItemPtr;
+
+struct FAdvancedRenamerPreviewListItem
+{
+	static FName OriginalNameColumnName;
+	static FName NewNameColumnName;
+
+	FAdvancedRenamerPreviewListItem(int32 InHash, const FString InOriginalName)
+		: Hash(InHash)
+		, OriginalName(InOriginalName)
+		, NewName(FString(""))
+	{
+	}
+
+	int32 Hash;
+	const FString OriginalName;
+	mutable FString NewName;
+};
+
+typedef TSharedPtr<FAdvancedRenamerPreviewListItem, ESPMode::ThreadSafe> FObjectRenamePreviewListItemPtr;
+typedef TWeakPtr<FAdvancedRenamerPreviewListItem, ESPMode::ThreadSafe> FObjectRenamePreviewListItemWeakPtr;
+
+class SAdvancedRenamerPreviewListRow : public SMultiColumnTableRow<FObjectRenamePreviewListItemPtr>
+{
+public:
+	SLATE_BEGIN_ARGS(SAdvancedRenamerPreviewListRow) { }
+	SLATE_END_ARGS()
+
+	void Construct(const FArguments& InArgs, TSharedPtr<SAdvancedRenamerPanel> InRenamePanel, 
+		const TSharedRef<STableViewBase>& InOwnerTableView, FObjectRenamePreviewListItemPtr InRowItem);
+
+	virtual TSharedRef<SWidget> GenerateWidgetForColumn(const FName& ColumnName) override;
+
+protected:
+
+	TWeakPtr<SAdvancedRenamerPanel> RenamePanel;
+	FObjectRenamePreviewListItemWeakPtr RowItem;
+
+};
+
+/**
+ * Implements its own provider interface so it can avoid long Execute_ lines and handle
+ * the 2 different types of provider (SharedPtr and UObject.)
+ */
+class SAdvancedRenamerPanel : public SCompoundWidget, private IAdvancedRenamerProvider
+{
+	friend class SAdvancedRenamerPreviewListRow;
+
+	SLATE_BEGIN_ARGS(SAdvancedRenamerPanel) 
+		: _SharedProvider(nullptr)
+		, _ObjectProvider(nullptr)
+		{}
+		SLATE_ARGUMENT(TSharedPtr<IAdvancedRenamerProvider>, SharedProvider)
+		SLATE_ARGUMENT(UObject*, ObjectProvider)
+	SLATE_END_ARGS()
+
+public:
+
+	void Construct(const FArguments& InArgs);
+
+protected:
+
+	static const inline double MinUpdateFrequency = 0.1;
+
+	TStrongObjectPtr<UObject> ObjectProvider;
+	TSharedPtr<IAdvancedRenamerProvider> SharedProvider;
+	TArray<TSharedPtr<FAdvancedRenamerPreviewListItem>> ListData;
+	TSharedPtr<FUICommandList> CommandList;
+
+	bool bValidNames = false;
+	bool bListNeedsUpdate = false;
+	bool bBaseName = false;
+	bool bPrefix = false;
+	bool bPrefixRemove = false;
+	bool bPrefixRemoveCharacters = false;
+	bool bSuffix = false;
+	bool bSuffixRemove = false;
+	bool bSuffixRemoveCharacters = false;
+	bool bSuffixRemoveNumber = false;
+	bool bSuffixNumber = false;
+	bool bSearchReplacePlainText = false;
+	bool bSearchReplaceRegex = false;
+	bool bSearchReplaceIgnoreCase = true;
+
+	double ListLastUpdateTime = 0;
+	float MinDesiredOriginalNameWidth = 0.f;
+	float MinDesiredNewNameWidth = 0.f;
+
+	uint8 PrefixRemoveCharacterCount = 0;
+	uint8 SuffixRemoveCharacterCount = 0;
+
+	TSharedPtr<SCheckBox> BaseNameCheckBox;
+	TSharedPtr<SEditableTextBox> BaseNameTextBox;
+	TSharedPtr<SCheckBox> PrefixCheckBox;
+	TSharedPtr<SEditableTextBox> PrefixTextBox;
+	TSharedPtr<SCheckBox> PrefixRemoveCheckBox;
+	TSharedPtr<SEditableTextBox> PrefixSeparatorTextBox;
+	TSharedPtr<SCheckBox> PrefixRemoveCharactersCheckBox;
+	TSharedPtr<SSpinBox<uint8>> PrefixRemoveCharactersSpinBox;
+	TSharedPtr<SCheckBox> SuffixCheckBox;
+	TSharedPtr<SEditableTextBox> SuffixTextBox;
+	TSharedPtr<SCheckBox> SuffixRemoveCheckBox;
+	TSharedPtr<SEditableTextBox> SuffixSeparatorTextBox;
+	TSharedPtr<SCheckBox> SuffixRemoveCharactersCheckBox;
+	TSharedPtr<SSpinBox<uint8>> SuffixRemoveCharactersSpinBox;
+	TSharedPtr<SCheckBox> SuffixRemoveNumberCheckBox;
+	TSharedPtr<SCheckBox> SuffixNumberCheckBox;
+	TSharedPtr<SSpinBox<int32>> SuffixNumberStartSpinBox;
+	TSharedPtr<SSpinBox<int32>> SuffixNumberStepSpinBox;
+	TSharedPtr<SCheckBox> SearchReplacePlainTextCheckbox;
+	TSharedPtr<SCheckBox> SearchReplaceRegexCheckbox;
+	TSharedPtr<SCheckBox> SearchReplaceIgnoreCaseCheckBox;
+	TSharedPtr<SMultiLineEditableTextBox> SearchReplaceSearchTextBox;
+	TSharedPtr<SMultiLineEditableTextBox> SearchReplaceReplaceTextBox;
+	TSharedPtr<SBox> RenamePreviewListBox;
+	TSharedPtr<SHeaderRow> RenamePreviewListHeaderRow;
+	TSharedPtr<SListView<FObjectRenamePreviewListItemPtr>> RenamePreviewList;
+	TSharedPtr<SButton> ApplyButton;
+
+	void CreateLeftPane(TSharedRef<SCanvas> Canvas);
+	TSharedRef<SWidget> CreateBaseName();
+	TSharedRef<SWidget> CreatePrefix();
+	TSharedRef<SWidget> CreateSuffix();
+	TSharedRef<SWidget> CreateSearchAndReplace();
+
+	void CreateRightPane(TSharedRef<SCanvas> Canvas);
+	TSharedRef<SWidget> CreateRenamePreview();
+
+	bool RenameObjects();
+
+	bool CloseWindow();
+
+	FString CreateNewName(int32 Index) const;
+	FString ApplyRename(const FString& OriginalName, int32 Index) const;
+	FString ApplyBaseName(const FString& OriginalName) const;
+	FString ApplyPrefix(const FString& OriginalName) const;
+	FString ApplySuffix(const FString& OriginalName, int32 Index) const;
+	FString ApplySearchPlainText(const FString& OriginalName) const;
+	FString ApplySearchReplaceRegex(const FString& OriginalName) const;
+
+	FString RegexReplace(const FString& OriginalString, const FRegexPattern Pattern, const FString& ReplaceString) const;
+
+	void UpdateEnables();
+	void RequestListViewRefresh();
+	void RefreshListView(const double InCurrentTime);
+	void UpdateRequiredListWidth();
+
+	void RemoveSelectedObjects();
+
+	virtual void Tick(const FGeometry& AllottedGeometry, const double InCurrentTime, const float InDeltaTime) override;
+
+	ECheckBoxState IsBaseNameChecked() const { return bBaseName ? ECheckBoxState::Checked : ECheckBoxState::Unchecked; }
+	void OnBaseNameCheckBoxChanged(ECheckBoxState NewState);
+
+	void OnBaseNameChanged(const FText& NewText);
+
+	ECheckBoxState IsPrefixChecked() const { return bPrefix ? ECheckBoxState::Checked : ECheckBoxState::Unchecked; }
+	void OnPrefixCheckBoxChanged(ECheckBoxState NewState);
+
+	void OnPrefixChanged(const FText& NewText);
+
+	ECheckBoxState IsPrefixRemoveChecked() const { return bPrefixRemove ? ECheckBoxState::Checked : ECheckBoxState::Unchecked; }
+	void OnPrefixRemoveCheckBoxChanged(ECheckBoxState NewState);
+
+	bool OnPrefixSeparatorVerifyTextChanged(const FText& InText, FText& OutErrorText) const;
+	void OnPrefixSeparatorChanged(const FText& NewText);
+
+	ECheckBoxState IsPrefixRemoveCharactersChecked() const { return bPrefixRemoveCharacters ? ECheckBoxState::Checked : ECheckBoxState::Unchecked; }
+	void OnPrefixRemoveCharactersCheckBoxChanged(ECheckBoxState NewState);
+
+	void OnPrefixRemoveCharactersChanged(uint8 NewValue);
+
+	ECheckBoxState IsSuffixChecked() const { return bSuffix ? ECheckBoxState::Checked : ECheckBoxState::Unchecked; }
+	void OnSuffixCheckBoxChanged(ECheckBoxState NewState);
+
+	void OnSuffixChanged(const FText& NewText);
+
+	ECheckBoxState IsSuffixRemoveChecked() const { return bSuffixRemove ? ECheckBoxState::Checked : ECheckBoxState::Unchecked; }
+	void OnSuffixRemoveCheckBoxChanged(ECheckBoxState NewState);
+
+	bool OnSuffixSeparatorVerifyTextChanged(const FText& InText, FText& OutErrorText) const;
+	void OnSuffixSeparatorChanged(const FText& NewText);
+
+	ECheckBoxState IsSuffixRemoveCharactersChecked() const { return bSuffixRemoveCharacters ? ECheckBoxState::Checked : ECheckBoxState::Unchecked; }
+	void OnSuffixRemoveCharactersCheckBoxChanged(ECheckBoxState NewState);
+
+	void OnSuffixRemoveCharactersChanged(uint8 NewValue);
+
+	ECheckBoxState IsSuffixRemoveNumberChecked() const { return bSuffixRemoveNumber ? ECheckBoxState::Checked : ECheckBoxState::Unchecked; }
+	void OnSuffixRemoveNumberCheckBoxChanged(ECheckBoxState NewState);
+
+	ECheckBoxState IsSuffixNumberChecked() const { return bSuffixNumber ? ECheckBoxState::Checked : ECheckBoxState::Unchecked; }
+	void OnSuffixNumberCheckBoxChanged(ECheckBoxState NewState);
+
+	void OnSuffixNumberStartChanged(int32 NewValue);
+
+	void OnSuffixNumberStepChanged(int32 NewValue);
+
+	ECheckBoxState IsSearchReplacePlainTextChecked() const { return bSearchReplacePlainText ? ECheckBoxState::Checked : ECheckBoxState::Unchecked; }
+	void OnSearchReplacePlainTextCheckBoxChanged(ECheckBoxState NewState);
+
+	ECheckBoxState IsSearchReplaceRegexChecked() const { return bSearchReplaceRegex ? ECheckBoxState::Checked : ECheckBoxState::Unchecked; }
+	void OnSearchReplaceRegexCheckBoxChanged(ECheckBoxState NewState);
+
+	ECheckBoxState IsSearchReplaceIgnoreCaseChecked() const { return bSearchReplaceIgnoreCase ? ECheckBoxState::Checked : ECheckBoxState::Unchecked; }
+	void OnSearchReplaceIgnoreCaseCheckBoxChanged(ECheckBoxState NewState);
+
+	void OnSearchReplaceSearchTextChanged(const FText& NewText);
+
+	void OnSearchReplaceReplaceTextChanged(const FText& NewText);
+
+	TSharedRef<ITableRow> OnGenerateRowForList(FObjectRenamePreviewListItemPtr Item, const TSharedRef<STableViewBase>& OwnerTable);
+	FReply OnListViewKeyDown(const FGeometry& MyGeometry, const FKeyEvent& KeyEvent);
+	TSharedPtr<SWidget> GenerateListViewContextMenu();
+
+	FReply OnApplyButtonClicked();
+
+	virtual int32 Num() const override;
+	virtual bool IsValidIndex(int32 Index) const override;
+	virtual uint32 GetHash(int32 Index) const override;
+	virtual FString GetOriginalName(int32 Index) const override;
+	virtual bool RemoveIndex(int32 Index) override;
+	virtual bool CanRename(int32 Index) const override;
+	virtual bool ExecuteRename(int32 Index, const FString& NewName) override;
+};
diff --git a/Engine/Plugins/Experimental/AdvancedRenamer/Source/AdvancedRenamer/Public/AdvancedRenamerBlueprintLibrary.h b/Engine/Plugins/Experimental/AdvancedRenamer/Source/AdvancedRenamer/Public/AdvancedRenamerBlueprintLibrary.h
new file mode 100644
--- /dev/null
+++ b/Engine/Plugins/Experimental/AdvancedRenamer/Source/AdvancedRenamer/Public/AdvancedRenamerBlueprintLibrary.h
@@ -0,0 +1,38 @@
+// Copyright Epic Games, Inc. All Rights Reserved.
+
+#pragma once
+
+#include "Kismet/BlueprintFunctionLibrary.h"
+#include "Containers/ContainersFwd.h"
+#include "Templates/SharedPointerFwd.h"
+#include "AdvancedRenamerBlueprintLibrary.generated.h"
+
+class AActor;
+class IAdvancedRenamerProvider;
+class IToolkitHost;
+class SWidget;
+class UUserWidget;
+
+UCLASS(BlueprintType, ClassGroup="Advanced Rename Panel")
+class ADVANCEDRENAMER_API UAdvancedRenamerBlueprintLibrary : public UBlueprintFunctionLibrary
+{
+	GENERATED_BODY()
+
+public:
+	UFUNCTION(BlueprintCallable, Category = "Advanced Renamer Panel", meta = (DisplayName = "Open Advanced Rename Panel"))
+	static void BP_OpenAdvancedRenamer(UObject* RenameProvider, UUserWidget* UserWidget = nullptr);
+
+	static void OpenAdvancedRenamer(TSharedRef<IAdvancedRenamerProvider> RenameProvider, TSharedPtr<IToolkitHost> ToolkitHost);
+
+	static void OpenAdvancedRenamer(TSharedRef<IAdvancedRenamerProvider> RenameProvider, TSharedPtr<SWidget> ParentWidget);
+
+	UFUNCTION(BlueprintCallable, Category = "Advanced Renamer Panel", meta = (DisplayName = "Open Advanced Rename Panel For Actors"))
+	static void BP_OpenAdvancedRenamerForActors(const TArray<AActor*>& Actors, UUserWidget* UserWidget = nullptr);
+
+	static void OpenAdvancedRenamerForActors(const TArray<AActor*>& Actors, TSharedPtr<IToolkitHost> ToolkitHost);
+
+	static void OpenAdvancedRenamerForActors(const TArray<AActor*>& Actors, TSharedPtr<SWidget> ParentWidget);
+
+	UFUNCTION(BlueprintCallable, Category = "Advanced Renamer Panel")
+	static TArray<AActor*> GetActorsSharingClassesInWorld(const TArray<AActor*>& InActors);
+};
diff --git a/Engine/Plugins/Experimental/AdvancedRenamer/Source/AdvancedRenamer/Public/AdvancedRenamerModule.h b/Engine/Plugins/Experimental/AdvancedRenamer/Source/AdvancedRenamer/Public/AdvancedRenamerModule.h
new file mode 100644
--- /dev/null
+++ b/Engine/Plugins/Experimental/AdvancedRenamer/Source/AdvancedRenamer/Public/AdvancedRenamerModule.h
@@ -0,0 +1,23 @@
+// Copyright Epic Games, Inc. All Rights Reserved.
+
+#pragma once
+
+#include "Modules/ModuleInterface.h"
+#include "Logging/LogMacros.h"
+
+DECLARE_LOG_CATEGORY_EXTERN(LogARP, Log, All);
+
+/**
+ * Advanced Rename Panel Plugin - Easily bulk rename stuff!
+ */
+class ADVANCEDRENAMER_API FAdvancedRenamerModule : public IModuleInterface
+{
+public:	
+	static FAdvancedRenamerModule& Get();
+
+	//~ Begin IAdvancedRenamerModule interface
+	virtual void StartupModule() override;
+	virtual void ShutdownModule() override;
+	//~ End IAdvancedRenamerModule interface
+};
+
diff --git a/Engine/Plugins/Experimental/AdvancedRenamer/Source/AdvancedRenamer/Public/Providers/AdvancedRenamerActorProvider.h b/Engine/Plugins/Experimental/AdvancedRenamer/Source/AdvancedRenamer/Public/Providers/AdvancedRenamerActorProvider.h
new file mode 100644
--- /dev/null
+++ b/Engine/Plugins/Experimental/AdvancedRenamer/Source/AdvancedRenamer/Public/Providers/AdvancedRenamerActorProvider.h
@@ -0,0 +1,37 @@
+// Copyright Epic Games, Inc. All Rights Reserved.
+
+#pragma once
+
+#include "AdvancedRenamerProviderInterface.h"
+#include "Containers/Array.h"
+#include "UObject/WeakObjectPtr.h"
+#include "UObject/WeakObjectPtrTemplates.h"
+
+class FString;
+class AActor;
+
+class ADVANCEDRENAMER_API FAdvancedRenamerActorProvider : public IAdvancedRenamerProvider
+{
+public:
+
+	FAdvancedRenamerActorProvider();
+	virtual ~FAdvancedRenamerActorProvider();
+
+	void SetActorList(const TArray<TWeakObjectPtr<AActor>>& InActorList);
+	void AddActorList(const TArray<TWeakObjectPtr<AActor>>& InActorList);
+	void AddActorData(AActor* InActor);
+	AActor* GetActor(int32 InIndex) const;
+
+	//~ Begin IAdvancedRenamerProvider
+	virtual int32 Num() const override;
+	virtual bool IsValidIndex(int32 InIndex) const override;
+	virtual uint32 GetHash(int32 InIndex) const override;;
+	virtual FString GetOriginalName(int32 InIndex) const override;
+	virtual bool RemoveIndex(int32 InIndex) override;
+	virtual bool CanRename(int32 InIndex) const override;
+	virtual bool ExecuteRename(int32 InIndex, const FString& InNewName) override;
+	//~ End IAdvancedRenamerProvider
+
+protected:
+	TArray<TWeakObjectPtr<AActor>> ActorList;
+};
diff --git a/Engine/Plugins/Experimental/AdvancedRenamer/Source/AdvancedRenamer/Public/Providers/AdvancedRenamerAssetProvider.h b/Engine/Plugins/Experimental/AdvancedRenamer/Source/AdvancedRenamer/Public/Providers/AdvancedRenamerAssetProvider.h
new file mode 100644
--- /dev/null
+++ b/Engine/Plugins/Experimental/AdvancedRenamer/Source/AdvancedRenamer/Public/Providers/AdvancedRenamerAssetProvider.h
@@ -0,0 +1,32 @@
+// Copyright Epic Games, Inc. All Rights Reserved.
+
+#pragma once
+
+#include "AdvancedRenamerProviderInterface.h"
+#include "AssetRegistry/AssetData.h"
+#include "Containers/Array.h"
+
+struct FAssetData;
+
+class ADVANCEDRENAMER_API FAdvancedRenamerAssetProvider : public IAdvancedRenamerProvider
+{
+public:
+	FAdvancedRenamerAssetProvider();
+	virtual ~FAdvancedRenamerAssetProvider() override;
+
+	void SetAssetList(const TArray<FAssetData>& InAssetList);
+	void AddAssetList(const TArray<FAssetData>& InAssetList);
+	void AddAssetData(const FAssetData& InAsset);
+	UObject* GetAsset(int32 Index) const;
+
+	virtual int32 Num() const override;
+	virtual bool IsValidIndex(int32 Index) const override;
+	virtual uint32 GetHash(int32 Index) const override;;
+	virtual FString GetOriginalName(int32 Index) const override;
+	virtual bool RemoveIndex(int32 Index) override;
+	virtual bool CanRename(int32 Index) const override;
+	virtual bool ExecuteRename(int32 Index, const FString& NewName) override;
+
+protected:
+	TArray<FAssetData> AssetList;
+};
diff --git a/Engine/Plugins/Experimental/AdvancedRenamer/Source/AdvancedRenamer/Public/Providers/AdvancedRenamerObjectAssetProvider.h b/Engine/Plugins/Experimental/AdvancedRenamer/Source/AdvancedRenamer/Public/Providers/AdvancedRenamerObjectAssetProvider.h
new file mode 100644
--- /dev/null
+++ b/Engine/Plugins/Experimental/AdvancedRenamer/Source/AdvancedRenamer/Public/Providers/AdvancedRenamerObjectAssetProvider.h
@@ -0,0 +1,49 @@
+// Copyright Epic Games, Inc. All Rights Reserved.
+
+#pragma once
+
+#include "AdvancedRenamerObjectProviderInterface.h"
+#include "AdvancedRenamerAssetProvider.h"
+#include "UObject/Object.h"
+#include "UObject/WeakObjectPtr.h"
+#include "UObject/WeakObjectPtrTemplates.h"
+#include "AdvancedRenamerObjectAssetProvider.generated.h"
+
+UCLASS(Abstract, BlueprintType, Blueprintable)
+class ADVANCEDRENAMER_API UAdvancedRenamerAssetProvider : public UObject, public IAdvancedRenamerObjectProvider
+#if CPP
+	, public FAdvancedRenamerAssetProvider
+#endif
+{
+	GENERATED_BODY()
+
+public:
+
+	UAdvancedRenamerAssetProvider()
+	{
+	}
+
+	UFUNCTION(BlueprintNativeEvent, BlueprintCallable, Category = "Advanced Rename Panel", meta = (DisplayName = "Set Asset List"))
+	void BP_SetAssetList(const TArray<UObject*>& InAssetList);
+	virtual void BP_SetAssetList_Implementation(const TArray<UObject*>& InAssetList);
+
+	UFUNCTION(BlueprintNativeEvent, BlueprintCallable, Category = "Advanced Rename Panel", meta = (DisplayName = "Add Asset List"))
+	void BP_AddAssetList(const TArray<UObject*>& InAssetList);
+	virtual void BP_AddAssetList_Implementation(const TArray<UObject*>& InAssetList);
+
+	UFUNCTION(BlueprintNativeEvent, BlueprintCallable, Category = "Advanced Rename Panel", meta = (DisplayName = "Add Asset Data"))
+	void BP_AddAssetData(UObject* InObject);
+	virtual void BP_AddAssetData_Implementation(UObject* InObject);
+
+	UFUNCTION(BlueprintNativeEvent, BlueprintCallable, Category = "Advanced Rename Panel", meta = (DisplayName = "Get Asset"))
+	UObject* BP_GetAsset(int32 Index) const;
+	virtual UObject* BP_GetAsset_Implementation(int32 Index) const;
+
+	virtual int32 Num() const override;
+	virtual bool IsValidIndex(int32 Index) const override;
+	virtual uint32 GetHash(int32 Index) const override;
+	virtual FString GetOriginalName(int32 Index) const override;
+	virtual bool RemoveIndex(int32 Index) override;
+	virtual bool CanRename(int32 Index) const override;
+	virtual bool ExecuteRename(int32 Index, const FString & NewName) override;
+};
diff --git a/Engine/Plugins/Experimental/AdvancedRenamer/Source/AdvancedRenamer/Public/Providers/AdvancedRenamerObjectObjectProvider.h b/Engine/Plugins/Experimental/AdvancedRenamer/Source/AdvancedRenamer/Public/Providers/AdvancedRenamerObjectObjectProvider.h
new file mode 100644
--- /dev/null
+++ b/Engine/Plugins/Experimental/AdvancedRenamer/Source/AdvancedRenamer/Public/Providers/AdvancedRenamerObjectObjectProvider.h
@@ -0,0 +1,48 @@
+// Copyright Epic Games, Inc. All Rights Reserved.
+
+#pragma once
+
+#include "AdvancedRenamerObjectProviderInterface.h"
+#include "AdvancedRenamerObjectProvider.h"
+#include "Containers/UnrealString.h"
+#include "UObject/Object.h"
+#include "UObject/WeakObjectPtr.h"
+#include "UObject/WeakObjectPtrTemplates.h"
+
+#include "AdvancedRenamerObjectObjectProvider.generated.h"
+
+UCLASS(Abstract, BlueprintType, Blueprintable)
+class ADVANCEDRENAMER_API UAdvancedRenamerObjectObjectProvider : public UObject, public IAdvancedRenamerObjectProvider
+#if CPP
+	, public FAdvancedRenamerObjectProvider
+#endif
+{
+	GENERATED_BODY()
+
+public:
+	UAdvancedRenamerObjectObjectProvider() = default;
+
+	UFUNCTION(BlueprintNativeEvent, BlueprintCallable, Category = "Advanced Rename Panel", meta = (DisplayName = "Set Object List"))
+	void BP_SetObjectList(const TArray<UObject*>& InObjectList);
+	virtual void BP_SetObjectList_Implementation(const TArray<UObject*>& InObjectList);
+
+	UFUNCTION(BlueprintNativeEvent, BlueprintCallable, Category = "Advanced Rename Panel", meta = (DisplayName = "Add Object List"))
+	void BP_AddObjectList(const TArray<UObject*>& InObjectList);
+	virtual void BP_AddObjectList_Implementation(const TArray<UObject*>& InObjectList);
+
+	UFUNCTION(BlueprintNativeEvent, BlueprintCallable, Category = "Advanced Rename Panel", meta = (DisplayName = "Add Object Data"))
+	void BP_AddObjectData(UObject* InObject);
+	virtual void BP_AddObjectData_Implementation(UObject* InObject);
+
+	UFUNCTION(BlueprintNativeEvent, BlueprintCallable, Category = "Advanced Rename Panel", meta = (DisplayName = "Get Object"))
+	UObject* BP_GetObject(int32 Index) const;
+	virtual UObject* BP_GetObject_Implementation(int32 Index) const;
+
+	virtual int32 Num() const override;
+	virtual bool IsValidIndex(int32 Index) const override;
+	virtual uint32 GetHash(int32 Index) const override;
+	virtual FString GetOriginalName(int32 Index) const override;
+	virtual bool RemoveIndex(int32 Index) override;
+	virtual bool CanRename(int32 Index) const override;
+	virtual bool ExecuteRename(int32 Index, const FString & NewName) override;
+};
diff --git a/Engine/Plugins/Experimental/AdvancedRenamer/Source/AdvancedRenamer/Public/Providers/AdvancedRenamerObjectProvider.h b/Engine/Plugins/Experimental/AdvancedRenamer/Source/AdvancedRenamer/Public/Providers/AdvancedRenamerObjectProvider.h
new file mode 100644
--- /dev/null
+++ b/Engine/Plugins/Experimental/AdvancedRenamer/Source/AdvancedRenamer/Public/Providers/AdvancedRenamerObjectProvider.h
@@ -0,0 +1,36 @@
+// Copyright Epic Games, Inc. All Rights Reserved.
+
+#pragma once
+
+#include "AdvancedRenamerProviderInterface.h"
+#include "Containers/Array.h"
+#include "UObject/WeakObjectPtr.h"
+#include "UObject/WeakObjectPtrTemplates.h"
+
+class FString;
+class UObject;
+
+class ADVANCEDRENAMER_API FAdvancedRenamerObjectProvider : public IAdvancedRenamerProvider
+{
+public:
+
+	FAdvancedRenamerObjectProvider();
+	virtual ~FAdvancedRenamerObjectProvider();
+
+	void SetObjectList(const TArray<TWeakObjectPtr<UObject>>& InObjectList);
+	void AddObjectList(const TArray<TWeakObjectPtr<UObject>>& InObjectList);
+	void AddObjectData(UObject* InObject);
+	UObject* GetObject(int32 Index) const;
+
+	virtual int32 Num() const override;
+	virtual bool IsValidIndex(int32 Index) const override;
+	virtual uint32 GetHash(int32 Index) const override;;
+	virtual FString GetOriginalName(int32 Index) const override;
+	virtual bool RemoveIndex(int32 Index) override;
+	virtual bool CanRename(int32 Index) const override;
+	virtual bool ExecuteRename(int32 Index, const FString& NewName) override;
+
+protected:
+
+	TArray<TWeakObjectPtr<UObject>> ObjectList;
+};
diff --git a/Engine/Plugins/Experimental/AdvancedRenamer/Source/AdvancedRenamer/Public/Providers/AdvancedRenamerObjectProviderInterface.h b/Engine/Plugins/Experimental/AdvancedRenamer/Source/AdvancedRenamer/Public/Providers/AdvancedRenamerObjectProviderInterface.h
new file mode 100644
--- /dev/null
+++ b/Engine/Plugins/Experimental/AdvancedRenamer/Source/AdvancedRenamer/Public/Providers/AdvancedRenamerObjectProviderInterface.h
@@ -0,0 +1,56 @@
+// Copyright Epic Games, Inc. All Rights Reserved.
+
+#pragma once
+
+#include "AdvancedRenamerProviderInterface.h"
+#include "UObject/Interface.h"
+#include "AdvancedRenamerObjectProviderInterface.generated.h"
+
+UINTERFACE()
+class ADVANCEDRENAMER_API UAdvancedRenamerObjectProvider : public UInterface
+{
+	GENERATED_BODY()
+};
+
+class ADVANCEDRENAMER_API IAdvancedRenamerObjectProvider 
+#if CPP
+	: public IAdvancedRenamerProvider
+#endif
+{
+	GENERATED_BODY()
+
+public:
+
+	/** 
+	 * Set as BNE so that there will be complaints when these are not implemented.
+	 * A BIE will silently fail.
+	 */	
+
+	UFUNCTION(BlueprintNativeEvent, BlueprintCallable, Category = "Advanced Rename Panel", meta = (DisplayName = "Num"))
+	int32 BP_Num() const;
+	virtual int32 BP_Num_Implementation() const PURE_VIRTUAL(IAdvancedRenamerObjectProvider::BP_Num, return 0;)
+
+	UFUNCTION(BlueprintNativeEvent, BlueprintCallable, Category = "Advanced Rename Panel", meta = (DisplayName = "IsValidIndex"))
+	bool BP_IsValidIndex(int32 Index) const;
+	virtual bool BP_IsValidIndex_Implementation(int32 Index) const PURE_VIRTUAL(IAdvancedRenamerObjectProvider::BP_IsValidIndex, return false;)
+
+	UFUNCTION(BlueprintNativeEvent, BlueprintCallable, Category = "Advanced Rename Panel", meta = (DisplayName = "GetHash"))
+	int32 BP_GetHash(int32 Index) const;
+	virtual int32 BP_GetHash_Implementation(int32 Index) const PURE_VIRTUAL(IAdvancedRenamerObjectProvider::BP_GetHash, return 0;)
+
+	UFUNCTION(BlueprintNativeEvent, BlueprintCallable, Category = "Advanced Rename Panel", meta = (DisplayName = "GetOriginalName"))
+	FString BP_GetOriginalName(int32 Index) const;
+	virtual FString BP_GetOriginalName_Implementation(int32 Index) const PURE_VIRTUAL(IAdvancedRenamerObjectProvider::BP_GetOriginalName, return "";)
+
+	UFUNCTION(BlueprintNativeEvent, BlueprintCallable, Category = "Advanced Rename Panel", meta = (DisplayName = "RemoveIndex"))
+	bool BP_RemoveIndex(int32 Index);
+	virtual bool BP_RemoveIndex_Implementation(int32 Index) PURE_VIRTUAL(IAdvancedRenamerObjectProvider::BP_RemoveIndex, return false;)
+
+	UFUNCTION(BlueprintNativeEvent, BlueprintCallable, Category = "Advanced Rename Panel", meta = (DisplayName = "CanRename"))
+	bool BP_CanRename(int32 Index) const;
+	virtual bool BP_CanRename_Implementation(int32 Index) const PURE_VIRTUAL(IAdvancedRenamerObjectProvider::BP_CanRename, return false;)
+
+	UFUNCTION(BlueprintNativeEvent, BlueprintCallable, Category = "Advanced Rename Panel", meta = (DisplayName = "Rename"))
+	bool BP_ExecuteRename(int32 Index, const FString& NewName);
+	virtual bool BP_ExecuteRename_Implementation(int32 Index, const FString& NewName) PURE_VIRTUAL(IAdvancedRenamerObjectProvider::BP_Rename, return false;)
+};
diff --git a/Engine/Plugins/Experimental/AdvancedRenamer/Source/AdvancedRenamer/Public/Providers/AdvancedRenamerProviderInterface.h b/Engine/Plugins/Experimental/AdvancedRenamer/Source/AdvancedRenamer/Public/Providers/AdvancedRenamerProviderInterface.h
new file mode 100644
--- /dev/null
+++ b/Engine/Plugins/Experimental/AdvancedRenamer/Source/AdvancedRenamer/Public/Providers/AdvancedRenamerProviderInterface.h
@@ -0,0 +1,20 @@
+// Copyright Epic Games, Inc. All Rights Reserved.
+
+#pragma once
+
+#include "HAL/Platform.h"
+
+class FString;
+
+class IAdvancedRenamerProvider
+{
+public:
+	virtual ~IAdvancedRenamerProvider() = default;
+	virtual int32 Num() const = 0;
+	virtual bool IsValidIndex(int32 Index) const = 0;
+	virtual uint32 GetHash(int32 Index) const = 0;
+	virtual FString GetOriginalName(int32 Index) const = 0;
+	virtual bool RemoveIndex(int32 Index) = 0;
+	virtual bool CanRename(int32 Index) const = 0;
+	virtual bool ExecuteRename(int32 Index, const FString& NewName) = 0;
+};
