Advanced Renamer
- Added analytics to opening the UI.

diff --git a/Engine/Plugins/Experimental/AdvancedRenamer/Source/AdvancedRenamer/Private/Slate/SAdvancedRenamerPanel.cpp b/Engine/Plugins/Experimental/AdvancedRenamer/Source/AdvancedRenamer/Private/Slate/SAdvancedRenamerPanel.cpp
--- a/Engine/Plugins/Experimental/AdvancedRenamer/Source/AdvancedRenamer/Private/Slate/SAdvancedRenamerPanel.cpp
+++ b/Engine/Plugins/Experimental/AdvancedRenamer/Source/AdvancedRenamer/Private/Slate/SAdvancedRenamerPanel.cpp
@@ -3,6 +3,7 @@
 #include "Slate/SAdvancedRenamerPanel.h"
 #include "AdvancedRenamerModule.h"
 #include "AdvancedRenamerStyle.h"
+#include "EngineAnalytics.h"
 #include "Framework/Commands/GenericCommands.h"
 #include "Framework/MultiBox/MultiBoxBuilder.h"
 #include "Internationalization/Regex.h"
@@ -177,6 +178,11 @@ void SAdvancedRenamerPanel::Construct(const FArguments& InArgs)
 		Canvas
 	];
 	// @formatter:on
+
+	if (FEngineAnalytics::IsAvailable())
+	{
+		FEngineAnalytics::GetProvider().RecordEvent(TEXT("Editor.Usage.AdvancedRenamer.Opened"));
+	}
 }
 
 void SAdvancedRenamerPanel::CreateLeftPane(TSharedRef<SCanvas> Canvas)
