Advanced Renamer
- Moved module header to private and removed Get method

diff --git a/Engine/Plugins/Experimental/AdvancedRenamer/Source/AdvancedRenamer/Private/AdvancedRenamerModule.cpp b/Engine/Plugins/Experimental/AdvancedRenamer/Source/AdvancedRenamer/Private/AdvancedRenamerModule.cpp
--- a/Engine/Plugins/Experimental/AdvancedRenamer/Source/AdvancedRenamer/Private/AdvancedRenamerModule.cpp
+++ b/Engine/Plugins/Experimental/AdvancedRenamer/Source/AdvancedRenamer/Private/AdvancedRenamerModule.cpp
@@ -9,12 +9,6 @@
 
 DEFINE_LOG_CATEGORY(LogARP);
 
-FAdvancedRenamerModule& FAdvancedRenamerModule::Get()
-{
-	static const FName ModuleName = TEXT("AdvancedRenamer");
-	return FModuleManager::LoadModuleChecked<FAdvancedRenamerModule>(ModuleName);
-}
-
 void FAdvancedRenamerModule::StartupModule()
 {
 	FAdvancedRenamerStyle::Initialize();
diff --git a/Engine/Plugins/Experimental/AdvancedRenamer/Source/AdvancedRenamer/Private/AdvancedRenamerModule.h b/Engine/Plugins/Experimental/AdvancedRenamer/Source/AdvancedRenamer/Private/AdvancedRenamerModule.h
new file mode 100644
--- /dev/null
+++ b/Engine/Plugins/Experimental/AdvancedRenamer/Source/AdvancedRenamer/Private/AdvancedRenamerModule.h
@@ -0,0 +1,21 @@
+// Copyright Epic Games, Inc. All Rights Reserved.
+
+#pragma once
+
+#include "Modules/ModuleInterface.h"
+#include "Logging/LogMacros.h"
+
+DECLARE_LOG_CATEGORY_EXTERN(LogARP, Log, All);
+
+/**
+ * Advanced Rename Panel Plugin - Easily bulk rename stuff!
+ */
+class ADVANCEDRENAMER_API FAdvancedRenamerModule : public IModuleInterface
+{
+public:	
+	//~ Begin IAdvancedRenamerModule interface
+	virtual void StartupModule() override;
+	virtual void ShutdownModule() override;
+	//~ End IAdvancedRenamerModule interface
+};
+
diff --git a/Engine/Plugins/Experimental/AdvancedRenamer/Source/AdvancedRenamer/Public/AdvancedRenamerModule.h b/Engine/Plugins/Experimental/AdvancedRenamer/Source/AdvancedRenamer/Public/AdvancedRenamerModule.h
deleted file mode 100644
