Advanced Renamer:
- Unexported unneeded APIs, moved advanced renamer bp library to module interface

#changelist validated
#virtualized

diff --git a/Engine/Plugins/Experimental/AdvancedRenamer/Source/AdvancedRenamer/AdvancedRenamer.Build.cs b/Engine/Plugins/Experimental/AdvancedRenamer/Source/AdvancedRenamer/AdvancedRenamer.Build.cs
--- a/Engine/Plugins/Experimental/AdvancedRenamer/Source/AdvancedRenamer/AdvancedRenamer.Build.cs
+++ b/Engine/Plugins/Experimental/AdvancedRenamer/Source/AdvancedRenamer/AdvancedRenamer.Build.cs
@@ -11,7 +11,7 @@ public class AdvancedRenamer : ModuleRules
 		PublicDependencyModuleNames.AddRange(
 			new string[]
 			{
-				"CoreUObject",
+				"Core",
 			}
 		);
 
@@ -19,7 +19,7 @@ public class AdvancedRenamer : ModuleRules
 			new string[]
 			{
 				"ContentBrowser",
-				"Core",
+				"CoreUObject",
 				"Engine",
 				"InputCore",
 				"LevelEditor",
diff --git a/Engine/Plugins/Experimental/AdvancedRenamer/Source/AdvancedRenamer/Private/AdvancedRenamerBlueprintLibrary.cpp b/Engine/Plugins/Experimental/AdvancedRenamer/Source/AdvancedRenamer/Private/AdvancedRenamerBlueprintLibrary.cpp
deleted file mode 100644
diff --git a/Engine/Plugins/Experimental/AdvancedRenamer/Source/AdvancedRenamer/Private/AdvancedRenamerModule.cpp b/Engine/Plugins/Experimental/AdvancedRenamer/Source/AdvancedRenamer/Private/AdvancedRenamerModule.cpp
--- a/Engine/Plugins/Experimental/AdvancedRenamer/Source/AdvancedRenamer/Private/AdvancedRenamerModule.cpp
+++ b/Engine/Plugins/Experimental/AdvancedRenamer/Source/AdvancedRenamer/Private/AdvancedRenamerModule.cpp
@@ -3,12 +3,32 @@
 #include "AdvancedRenamerModule.h"
 #include "AdvancedRenamerCommands.h"
 #include "AdvancedRenamerStyle.h"
+#include "EngineUtils.h"
+#include "Framework/Application/SlateApplication.h"
 #include "Integrations/AdvancedRenamerContentBrowserIntegration.h"
 #include "Integrations/AdvancedRenamerLevelEditorIntegration.h"
-#include "Modules/ModuleManager.h"
+#include "Providers/AdvancedRenamerActorProvider.h"
+#include "Slate/SAdvancedRenamerPanel.h"
+#include "Toolkits/IToolkitHost.h"
+#include "Widgets/SWindow.h"
 
 DEFINE_LOG_CATEGORY(LogARP);
 
+#define LOCTEXT_NAMESPACE "AdvancedRenamerModule"
+
+namespace UE::AdvancedRenamer::Private
+{
+	TSharedRef<SWindow> CreateAdvancedRenamerWindow()
+	{
+		return SNew(SWindow)
+			.Title(LOCTEXT("AdvancedRenameWindow", "Rename Actors"))
+			.ClientSize(FVector2D(600.0f, 500.0f))
+			.SizingRule(ESizingRule::FixedSize)
+			.SupportsMaximize(false)
+			.SupportsMinimize(false);
+	}
+}
+
 void FAdvancedRenamerModule::StartupModule()
 {
 	FAdvancedRenamerStyle::Initialize();
@@ -25,4 +45,146 @@ void FAdvancedRenamerModule::ShutdownModule()
 	FAdvancedRenamerLevelEditorIntegration::Shutdown();
 }
 
+void FAdvancedRenamerModule::OpenAdvancedRenamer(const TSharedRef<IAdvancedRenamerProvider>& InRenameProvider, const TSharedPtr<SWidget>& InParentWidget)
+{
+	TSharedRef<SWindow> AdvancedRenameWindow = UE::AdvancedRenamer::Private::CreateAdvancedRenamerWindow();
+	AdvancedRenameWindow->SetContent(SNew(SAdvancedRenamerPanel).SharedProvider(InRenameProvider));
+
+	TSharedPtr<SWidget> ParentWindow = FSlateApplication::Get().FindBestParentWindowForDialogs(InParentWidget);
+	FSlateApplication::Get().AddModalWindow(AdvancedRenameWindow, ParentWindow);
+}
+
+void FAdvancedRenamerModule::OpenAdvancedRenamer(const TSharedRef<IAdvancedRenamerProvider>& InRenameProvider, const TSharedPtr<IToolkitHost>& InToolkitHost)
+{
+	if (InToolkitHost.IsValid())
+	{
+		OpenAdvancedRenamer(InRenameProvider, InToolkitHost->GetParentWidget());
+	}
+}
+
+void FAdvancedRenamerModule::OpenAdvancedRenamerForActors(const TArray<AActor*>& InActors, const TSharedPtr<SWidget>& InParentWidget)
+{
+	TArray<TWeakObjectPtr<AActor>> WeakObjects;
+
+	WeakObjects.Reserve(InActors.Num());
+
+	Algo::Transform(
+		InActors,
+		WeakObjects,
+		[](AActor* InActor)
+		{
+			return TWeakObjectPtr<AActor>(InActor);
+		}
+	);
+
+	TSharedRef<FAdvancedRenamerActorProvider> ActorProvider = MakeShared<FAdvancedRenamerActorProvider>();
+	ActorProvider->SetActorList(WeakObjects);
+
+	OpenAdvancedRenamer(ActorProvider, InParentWidget);
+}
+
+void FAdvancedRenamerModule::OpenAdvancedRenamerForActors(const TArray<AActor*>& InActors, const TSharedPtr<IToolkitHost>& InToolkitHost)
+{
+	if (InToolkitHost.IsValid())
+	{
+		OpenAdvancedRenamerForActors(InActors, InToolkitHost->GetParentWidget());
+	}
+}
+
+TArray<AActor*> FAdvancedRenamerModule::GetActorsSharingClassesInWorld(const TArray<AActor*>& InActors)
+{
+	TSet<UClass*> SelectedClasses;
+	bool bHasActorClass = false;
+	UWorld* World = nullptr;
+
+	// Scan selected items and add valid classes to the selected classes list.
+	for (AActor* SelectedActor : InActors)
+	{
+		if (!IsValid(SelectedActor))
+		{
+			continue;
+		}
+
+		if (!World)
+		{
+			World = SelectedActor->GetWorld();
+
+			if (!World)
+			{
+				break;
+			}
+		}
+
+		UClass* ActorClass = SelectedActor->GetClass();
+
+		/**
+		 * If we have a default AActor selected then all actors in the world share a
+		 * class with the selected actors. We don't need anything other than the AActor
+		 * class to get matches. Empty the array, store that and move on.
+		 */
+		if (ActorClass == AActor::StaticClass())
+		{
+			bHasActorClass = true;
+			SelectedClasses.Empty();
+			break;
+		}
+
+		SelectedClasses.Add(ActorClass);
+	}
+
+	if (!World)
+	{
+		return InActors;
+	}
+
+	TArray<UClass*> NonInheritingActorClasses;
+
+	if (bHasActorClass)
+	{
+		NonInheritingActorClasses.Add(AActor::StaticClass());
+	}
+	else
+	{
+		for (UClass* ActorClass : SelectedClasses)
+		{
+			bool bFoundParent = false;
+
+			for (UClass* ActorClassCheck : SelectedClasses)
+			{
+				if (ActorClass == ActorClassCheck)
+				{
+					continue;
+				}
+
+				if (ActorClass->IsChildOf(ActorClassCheck))
+				{
+					bFoundParent = true;
+					break;
+				}
+			}
+
+			if (!bFoundParent)
+			{
+				NonInheritingActorClasses.Add(ActorClass);
+			}
+		}
+	}
+
+	// Create outliner items for all the items matching the class list that are renameable.
+	TArray<AActor*> AllActors;
+	AllActors.Reserve(InActors.Num());
+
+	for (UClass* ActorClass : NonInheritingActorClasses)
+	{
+		for (AActor* Actor : TActorRange<AActor>(World, ActorClass))
+		{
+			AllActors.Add(Actor);
+		}
+	}
+
+	return AllActors;
+}
+
 IMPLEMENT_MODULE(FAdvancedRenamerModule, AdvancedRenamer)
+
+#undef LOCTEXT_NAMESPACE
diff --git a/Engine/Plugins/Experimental/AdvancedRenamer/Source/AdvancedRenamer/Private/AdvancedRenamerModule.h b/Engine/Plugins/Experimental/AdvancedRenamer/Source/AdvancedRenamer/Private/AdvancedRenamerModule.h
--- a/Engine/Plugins/Experimental/AdvancedRenamer/Source/AdvancedRenamer/Private/AdvancedRenamerModule.h
+++ b/Engine/Plugins/Experimental/AdvancedRenamer/Source/AdvancedRenamer/Private/AdvancedRenamerModule.h
@@ -2,20 +2,24 @@
 
 #pragma once
 
-#include "Modules/ModuleInterface.h"
+#include "IAdvancedRenamerModule.h"
 #include "Logging/LogMacros.h"
 
 DECLARE_LOG_CATEGORY_EXTERN(LogARP, Log, All);
 
-/**
- * Advanced Rename Panel Plugin - Easily bulk rename stuff!
- */
-class ADVANCEDRENAMER_API FAdvancedRenamerModule : public IModuleInterface
+/** Advanced Rename Panel Plugin - Easily bulk rename stuff! */
+class FAdvancedRenamerModule : public IAdvancedRenamerModule
 {
-public:	
-	//~ Begin IAdvancedRenamerModule interface
+	//~ Begin IModuleInterface
 	virtual void StartupModule() override;
 	virtual void ShutdownModule() override;
-	//~ End IAdvancedRenamerModule interface
-};
+	//~ End IModuleInterface
 
+	//~ Begin IAdvancedRenamerModule
+	virtual void OpenAdvancedRenamer(const TSharedRef<IAdvancedRenamerProvider>& InRenameProvider, const TSharedPtr<SWidget>& InParentWidget) override;
+	virtual void OpenAdvancedRenamer(const TSharedRef<IAdvancedRenamerProvider>& InRenameProvider, const TSharedPtr<IToolkitHost>& InToolkitHost) override;
+	virtual void OpenAdvancedRenamerForActors(const TArray<AActor*>& InActors, const TSharedPtr<SWidget>& InParentWidget) override;
+	virtual void OpenAdvancedRenamerForActors(const TArray<AActor*>& InActors, const TSharedPtr<IToolkitHost>& InToolkitHost) override;
+	virtual TArray<AActor*> GetActorsSharingClassesInWorld(const TArray<AActor*>& InActors) override;
+	//~ End IAdvancedRenamerModule
+};
diff --git a/Engine/Plugins/Experimental/AdvancedRenamer/Source/AdvancedRenamer/Private/Integrations/AdvancedRenamerContentBrowserIntegration.cpp b/Engine/Plugins/Experimental/AdvancedRenamer/Source/AdvancedRenamer/Private/Integrations/AdvancedRenamerContentBrowserIntegration.cpp
--- a/Engine/Plugins/Experimental/AdvancedRenamer/Source/AdvancedRenamer/Private/Integrations/AdvancedRenamerContentBrowserIntegration.cpp
+++ b/Engine/Plugins/Experimental/AdvancedRenamer/Source/AdvancedRenamer/Private/Integrations/AdvancedRenamerContentBrowserIntegration.cpp
@@ -1,13 +1,13 @@
 // Copyright Epic Games, Inc. All Rights Reserved.
 
 #include "Integrations/AdvancedRenamerContentBrowserIntegration.h"
-#include "AdvancedRenamerBlueprintLibrary.h"
 #include "ContentBrowserDelegates.h"
 #include "ContentBrowserModule.h"
 #include "Delegates/IDelegateInstance.h"
 #include "Framework/Commands/UIAction.h"
 #include "Framework/MultiBox/MultiBoxBuilder.h"
 #include "Framework/MultiBox/MultiBoxExtender.h"
+#include "IAdvancedRenamerModule.h"
 #include "Providers/AdvancedRenamerAssetProvider.h"
 
 #define LOCTEXT_NAMESPACE "AdvancedRenamerContentBrowserIntegration"
@@ -23,7 +23,7 @@ namespace UE::AdvancedRenamer::Private
 
 		TSharedPtr<SWidget> HostWidget = nullptr;
 
-		UAdvancedRenamerBlueprintLibrary::OpenAdvancedRenamer(StaticCastSharedRef<IAdvancedRenamerProvider>(AssetProvider), HostWidget);
+		IAdvancedRenamerModule::Get().OpenAdvancedRenamer(StaticCastSharedRef<IAdvancedRenamerProvider>(AssetProvider), HostWidget);
 	}
 
 	void ExtendAssetMenu(FMenuBuilder& MenuBuilder, const TArray<FAssetData> AssetArray)
diff --git a/Engine/Plugins/Experimental/AdvancedRenamer/Source/AdvancedRenamer/Private/Integrations/AdvancedRenamerLevelEditorIntegration.cpp b/Engine/Plugins/Experimental/AdvancedRenamer/Source/AdvancedRenamer/Private/Integrations/AdvancedRenamerLevelEditorIntegration.cpp
--- a/Engine/Plugins/Experimental/AdvancedRenamer/Source/AdvancedRenamer/Private/Integrations/AdvancedRenamerLevelEditorIntegration.cpp
+++ b/Engine/Plugins/Experimental/AdvancedRenamer/Source/AdvancedRenamer/Private/Integrations/AdvancedRenamerLevelEditorIntegration.cpp
@@ -1,12 +1,12 @@
 // Copyright Epic Games, Inc. All Rights Reserved.
 
 #include "Integrations/AdvancedRenamerLevelEditorIntegration.h"
-#include "AdvancedRenamerBlueprintLibrary.h"
 #include "AdvancedRenamerCommands.h"
 #include "Containers/Array.h"
 #include "Delegates/IDelegateInstance.h"
 #include "EditorModeManager.h"
 #include "GameFramework/Actor.h"
+#include "IAdvancedRenamerModule.h"
 #include "ILevelEditor.h"
 #include "LevelEditor.h"
 #include "Selection.h"
@@ -64,7 +64,7 @@ namespace UE::AdvancedRenamer::Private
 
 		TArray<AActor*> SelectedActors = GetSelectedActors(LevelEditor.ToSharedRef());
 
-		UAdvancedRenamerBlueprintLibrary::OpenAdvancedRenamerForActors(SelectedActors, StaticCastSharedPtr<IToolkitHost>(LevelEditor));
+		IAdvancedRenamerModule::Get().OpenAdvancedRenamerForActors(SelectedActors, StaticCastSharedPtr<IToolkitHost>(LevelEditor));
 	}
 
 	void RenameSharedClassActors(TWeakPtr<ILevelEditor> InLevelEditorWeak)
@@ -76,10 +76,12 @@ namespace UE::AdvancedRenamer::Private
 			return;
 		}
 
+		IAdvancedRenamerModule& AdvancedRenamerModule = IAdvancedRenamerModule::Get();
+
 		TArray<AActor*> SelectedActors = GetSelectedActors(LevelEditor.ToSharedRef());
-		SelectedActors = UAdvancedRenamerBlueprintLibrary::GetActorsSharingClassesInWorld(SelectedActors);
+		SelectedActors = AdvancedRenamerModule.GetActorsSharingClassesInWorld(SelectedActors);
 
-		UAdvancedRenamerBlueprintLibrary::OpenAdvancedRenamerForActors(SelectedActors, StaticCastSharedPtr<IToolkitHost>(LevelEditor));
+		AdvancedRenamerModule.OpenAdvancedRenamerForActors(SelectedActors, StaticCastSharedPtr<IToolkitHost>(LevelEditor));
 	}
 
 	void OnLevelEditorCreated(TSharedPtr<ILevelEditor> InLevelEditor)
diff --git a/Engine/Plugins/Experimental/AdvancedRenamer/Source/AdvancedRenamer/Private/Providers/AdvancedRenamerActorProvider.h b/Engine/Plugins/Experimental/AdvancedRenamer/Source/AdvancedRenamer/Private/Providers/AdvancedRenamerActorProvider.h
new file mode 100644
--- /dev/null
+++ b/Engine/Plugins/Experimental/AdvancedRenamer/Source/AdvancedRenamer/Private/Providers/AdvancedRenamerActorProvider.h
@@ -0,0 +1,36 @@
+// Copyright Epic Games, Inc. All Rights Reserved.
+
+#pragma once
+
+#include "Providers/IAdvancedRenamerProvider.h"
+#include "Containers/Array.h"
+#include "UObject/WeakObjectPtr.h"
+#include "UObject/WeakObjectPtrTemplates.h"
+
+class FString;
+class AActor;
+
+class FAdvancedRenamerActorProvider : public IAdvancedRenamerProvider
+{
+public:
+	FAdvancedRenamerActorProvider();
+	virtual ~FAdvancedRenamerActorProvider() override;
+
+	void SetActorList(const TArray<TWeakObjectPtr<AActor>>& InActorList);
+	void AddActorList(const TArray<TWeakObjectPtr<AActor>>& InActorList);
+	void AddActorData(AActor* InActor);
+	AActor* GetActor(int32 InIndex) const;
+
+protected:
+	//~ Begin IAdvancedRenamerProvider
+	virtual int32 Num() const override;
+	virtual bool IsValidIndex(int32 InIndex) const override;
+	virtual uint32 GetHash(int32 InIndex) const override;;
+	virtual FString GetOriginalName(int32 InIndex) const override;
+	virtual bool RemoveIndex(int32 InIndex) override;
+	virtual bool CanRename(int32 InIndex) const override;
+	virtual bool ExecuteRename(int32 InIndex, const FString& InNewName) override;
+	//~ End IAdvancedRenamerProvider
+
+	TArray<TWeakObjectPtr<AActor>> ActorList;
+};
diff --git a/Engine/Plugins/Experimental/AdvancedRenamer/Source/AdvancedRenamer/Private/Providers/AdvancedRenamerAssetProvider.h b/Engine/Plugins/Experimental/AdvancedRenamer/Source/AdvancedRenamer/Private/Providers/AdvancedRenamerAssetProvider.h
new file mode 100644
--- /dev/null
+++ b/Engine/Plugins/Experimental/AdvancedRenamer/Source/AdvancedRenamer/Private/Providers/AdvancedRenamerAssetProvider.h
@@ -0,0 +1,34 @@
+// Copyright Epic Games, Inc. All Rights Reserved.
+
+#pragma once
+
+#include "Providers/IAdvancedRenamerProvider.h"
+#include "AssetRegistry/AssetData.h"
+#include "Containers/Array.h"
+
+struct FAssetData;
+
+class FAdvancedRenamerAssetProvider : public IAdvancedRenamerProvider
+{
+public:
+	FAdvancedRenamerAssetProvider();
+	virtual ~FAdvancedRenamerAssetProvider() override;
+
+	void SetAssetList(const TArray<FAssetData>& InAssetList);
+	void AddAssetList(const TArray<FAssetData>& InAssetList);
+	void AddAssetData(const FAssetData& InAsset);
+	UObject* GetAsset(int32 Index) const;
+
+protected:
+	//~ Begin IAdvancedRenamerProvider
+	virtual int32 Num() const override;
+	virtual bool IsValidIndex(int32 Index) const override;
+	virtual uint32 GetHash(int32 Index) const override;;
+	virtual FString GetOriginalName(int32 Index) const override;
+	virtual bool RemoveIndex(int32 Index) override;
+	virtual bool CanRename(int32 Index) const override;
+	virtual bool ExecuteRename(int32 Index, const FString& NewName) override;
+	//~ End IAdvancedRenamerProvider
+
+	TArray<FAssetData> AssetList;
+};
diff --git a/Engine/Plugins/Experimental/AdvancedRenamer/Source/AdvancedRenamer/Private/Providers/AdvancedRenamerObjectAssetProvider.cpp b/Engine/Plugins/Experimental/AdvancedRenamer/Source/AdvancedRenamer/Private/Providers/AdvancedRenamerObjectAssetProvider.cpp
deleted file mode 100644
diff --git a/Engine/Plugins/Experimental/AdvancedRenamer/Source/AdvancedRenamer/Private/Providers/AdvancedRenamerObjectObjectProvider.cpp b/Engine/Plugins/Experimental/AdvancedRenamer/Source/AdvancedRenamer/Private/Providers/AdvancedRenamerObjectObjectProvider.cpp
deleted file mode 100644
diff --git a/Engine/Plugins/Experimental/AdvancedRenamer/Source/AdvancedRenamer/Private/Providers/AdvancedRenamerObjectProvider.h b/Engine/Plugins/Experimental/AdvancedRenamer/Source/AdvancedRenamer/Private/Providers/AdvancedRenamerObjectProvider.h
new file mode 100644
--- /dev/null
+++ b/Engine/Plugins/Experimental/AdvancedRenamer/Source/AdvancedRenamer/Private/Providers/AdvancedRenamerObjectProvider.h
@@ -0,0 +1,36 @@
+// Copyright Epic Games, Inc. All Rights Reserved.
+
+#pragma once
+
+#include "Providers/IAdvancedRenamerProvider.h"
+#include "Containers/Array.h"
+#include "UObject/WeakObjectPtr.h"
+#include "UObject/WeakObjectPtrTemplates.h"
+
+class FString;
+class UObject;
+
+class FAdvancedRenamerObjectProvider : public IAdvancedRenamerProvider
+{
+public:
+	FAdvancedRenamerObjectProvider();
+	virtual ~FAdvancedRenamerObjectProvider() override;
+
+	void SetObjectList(const TArray<TWeakObjectPtr<UObject>>& InObjectList);
+	void AddObjectList(const TArray<TWeakObjectPtr<UObject>>& InObjectList);
+	void AddObjectData(UObject* InObject);
+	UObject* GetObject(int32 Index) const;
+
+protected:
+	//~ Begin IAdvancedRenamerProvider
+	virtual int32 Num() const override;
+	virtual bool IsValidIndex(int32 Index) const override;
+	virtual uint32 GetHash(int32 Index) const override;;
+	virtual FString GetOriginalName(int32 Index) const override;
+	virtual bool RemoveIndex(int32 Index) override;
+	virtual bool CanRename(int32 Index) const override;
+	virtual bool ExecuteRename(int32 Index, const FString& NewName) override;
+	//~ End IAdvancedRenamerProvider
+
+	TArray<TWeakObjectPtr<UObject>> ObjectList;
+};
diff --git a/Engine/Plugins/Experimental/AdvancedRenamer/Source/AdvancedRenamer/Private/Slate/SAdvancedRenamerPanel.cpp b/Engine/Plugins/Experimental/AdvancedRenamer/Source/AdvancedRenamer/Private/Slate/SAdvancedRenamerPanel.cpp
--- a/Engine/Plugins/Experimental/AdvancedRenamer/Source/AdvancedRenamer/Private/Slate/SAdvancedRenamerPanel.cpp
+++ b/Engine/Plugins/Experimental/AdvancedRenamer/Source/AdvancedRenamer/Private/Slate/SAdvancedRenamerPanel.cpp
@@ -7,8 +7,7 @@
 #include "Framework/Commands/GenericCommands.h"
 #include "Framework/MultiBox/MultiBoxBuilder.h"
 #include "Internationalization/Regex.h"
-#include "Providers/AdvancedRenamerObjectProviderInterface.h"
-#include "Providers/AdvancedRenamerProviderInterface.h"
+#include "Providers/IAdvancedRenamerProvider.h"
 #include "ScopedTransaction.h"
 #include "Styling/AppStyle.h"
 #include "Styling/StyleColors.h"
@@ -116,12 +115,7 @@ TSharedRef<SWidget> SAdvancedRenamerPreviewListRow::GenerateWidgetForColumn(cons
 
 void SAdvancedRenamerPanel::Construct(const FArguments& InArgs)
 {
-	if (InArgs._ObjectProvider != nullptr)
-	{
-		ObjectProvider = TStrongObjectPtr<UObject>(InArgs._ObjectProvider);
-		check(ObjectProvider->GetClass()->ImplementsInterface(UAdvancedRenamerObjectProvider::StaticClass()));
-	}
-	else if (InArgs._SharedProvider.IsValid())
+	if (InArgs._SharedProvider.IsValid())
 	{
 		SharedProvider = InArgs._SharedProvider;
 	}
@@ -1785,50 +1779,22 @@ FReply SAdvancedRenamerPanel::OnApplyButtonClicked()
 
 int32 SAdvancedRenamerPanel::Num() const
 {
-	if (ObjectProvider.IsValid())
-	{
-		return IAdvancedRenamerObjectProvider::Execute_BP_Num(ObjectProvider.Get());
-	}
-	else
-	{
-		return SharedProvider->Num();
-	}
+	return SharedProvider->Num();
 }
 
 bool SAdvancedRenamerPanel::IsValidIndex(int32 Index) const
 {
-	if (ObjectProvider.IsValid())
-	{
-		return IAdvancedRenamerObjectProvider::Execute_BP_IsValidIndex(ObjectProvider.Get(), Index);
-	}
-	else
-	{
-		return SharedProvider->IsValidIndex(Index);
-	}
+	return SharedProvider->IsValidIndex(Index);
 }
 
 uint32 SAdvancedRenamerPanel::GetHash(int32 Index) const
 {
-	if (ObjectProvider.IsValid())
-	{
-		return static_cast<uint32>(IAdvancedRenamerObjectProvider::Execute_BP_GetHash(ObjectProvider.Get(), Index));
-	}
-	else
-	{
-		return SharedProvider->GetHash(Index);
-	}
+	return SharedProvider->GetHash(Index);
 }
 
 FString SAdvancedRenamerPanel::GetOriginalName(int32 Index) const
 {
-	if (ObjectProvider.IsValid())
-	{
-		return IAdvancedRenamerObjectProvider::Execute_BP_GetOriginalName(ObjectProvider.Get(), Index);
-	}
-	else
-	{
-		return SharedProvider->GetOriginalName(Index);
-	}
+	return SharedProvider->GetOriginalName(Index);
 }
 
 bool SAdvancedRenamerPanel::RemoveIndex(int32 Index)
@@ -1840,38 +1806,17 @@ bool SAdvancedRenamerPanel::RemoveIndex(int32 Index)
 		ListData.RemoveAt(Index);
 	}
 
-	if (ObjectProvider.IsValid())
-	{
-		return IAdvancedRenamerObjectProvider::Execute_BP_RemoveIndex(ObjectProvider.Get(), Index);
-	}
-	else
-	{
-		return SharedProvider->RemoveIndex(Index);
-	}
+	return SharedProvider->RemoveIndex(Index);
 }
 
 bool SAdvancedRenamerPanel::CanRename(int32 Index) const
 {
-	if (ObjectProvider.IsValid())
-	{
-		return IAdvancedRenamerObjectProvider::Execute_BP_CanRename(ObjectProvider.Get(), Index);
-	}
-	else
-	{
-		return SharedProvider->CanRename(Index);
-	}
+	return SharedProvider->CanRename(Index);
 }
 
 bool SAdvancedRenamerPanel::ExecuteRename(int32 Index, const FString& NewName)
 {
-	if (ObjectProvider.IsValid())
-	{
-		return IAdvancedRenamerObjectProvider::Execute_BP_ExecuteRename(ObjectProvider.Get(), Index, NewName);
-	}
-	else
-	{
-		return SharedProvider->ExecuteRename(Index, NewName);
-	}
+	return SharedProvider->ExecuteRename(Index, NewName);
 }
 
 #undef LOCTEXT_NAMESPACE
diff --git a/Engine/Plugins/Experimental/AdvancedRenamer/Source/AdvancedRenamer/Private/Slate/SAdvancedRenamerPanel.h b/Engine/Plugins/Experimental/AdvancedRenamer/Source/AdvancedRenamer/Private/Slate/SAdvancedRenamerPanel.h
--- a/Engine/Plugins/Experimental/AdvancedRenamer/Source/AdvancedRenamer/Private/Slate/SAdvancedRenamerPanel.h
+++ b/Engine/Plugins/Experimental/AdvancedRenamer/Source/AdvancedRenamer/Private/Slate/SAdvancedRenamerPanel.h
@@ -2,18 +2,16 @@
 
 #pragma once
 
-#include "Providers/AdvancedRenamerProviderInterface.h"
+#include "Providers/IAdvancedRenamerProvider.h"
 #include "Styling/SlateTypes.h"
 #include "UObject/StrongObjectPtr.h"
 #include "Widgets/Views/SListView.h"
 #include "Widgets/Views/STableRow.h"
 #include "Widgets/Views/STableViewBase.h"
 
-class FBlueprintEditor;
 class FRegexPattern;
 class FUICommandList;
 class IAdvancedRenamerProvider;
-class IAvaOutlinerItem;
 class SAdvancedRenamerPanel;
 class SBox;
 class SButton;
@@ -22,10 +20,8 @@ class SCheckBox;
 class SEditableTextBox;
 class SHeaderRow;
 class SMultiLineEditableTextBox; 
-template<typename NumericType>
-class SSpinBox;
 class UObject;
-typedef TSharedPtr<IAvaOutlinerItem> FAvaOutlinerItemPtr;
+template<typename NumericType> class SSpinBox;
 
 struct FAdvancedRenamerPreviewListItem
 {
@@ -50,7 +46,7 @@ typedef TWeakPtr<FAdvancedRenamerPreviewListItem, ESPMode::ThreadSafe> FObjectRe
 class SAdvancedRenamerPreviewListRow : public SMultiColumnTableRow<FObjectRenamePreviewListItemPtr>
 {
 public:
-	SLATE_BEGIN_ARGS(SAdvancedRenamerPreviewListRow) { }
+	SLATE_BEGIN_ARGS(SAdvancedRenamerPreviewListRow) {}
 	SLATE_END_ARGS()
 
 	void Construct(const FArguments& InArgs, TSharedPtr<SAdvancedRenamerPanel> InRenamePanel, 
@@ -73,23 +69,16 @@ class SAdvancedRenamerPanel : public SCompoundWidget, private IAdvancedRenamerPr
 {
 	friend class SAdvancedRenamerPreviewListRow;
 
-	SLATE_BEGIN_ARGS(SAdvancedRenamerPanel) 
-		: _SharedProvider(nullptr)
-		, _ObjectProvider(nullptr)
-		{}
+public:
+	SLATE_BEGIN_ARGS(SAdvancedRenamerPanel) {}
 		SLATE_ARGUMENT(TSharedPtr<IAdvancedRenamerProvider>, SharedProvider)
-		SLATE_ARGUMENT(UObject*, ObjectProvider)
 	SLATE_END_ARGS()
 
-public:
-
 	void Construct(const FArguments& InArgs);
 
 protected:
-
 	static const inline double MinUpdateFrequency = 0.1;
 
-	TStrongObjectPtr<UObject> ObjectProvider;
 	TSharedPtr<IAdvancedRenamerProvider> SharedProvider;
 	TArray<TSharedPtr<FAdvancedRenamerPreviewListItem>> ListData;
 	TSharedPtr<FUICommandList> CommandList;
diff --git a/Engine/Plugins/Experimental/AdvancedRenamer/Source/AdvancedRenamer/Public/AdvancedRenamerBlueprintLibrary.h b/Engine/Plugins/Experimental/AdvancedRenamer/Source/AdvancedRenamer/Public/AdvancedRenamerBlueprintLibrary.h
deleted file mode 100644
diff --git a/Engine/Plugins/Experimental/AdvancedRenamer/Source/AdvancedRenamer/Public/IAdvancedRenamerModule.h b/Engine/Plugins/Experimental/AdvancedRenamer/Source/AdvancedRenamer/Public/IAdvancedRenamerModule.h
new file mode 100644
--- /dev/null
+++ b/Engine/Plugins/Experimental/AdvancedRenamer/Source/AdvancedRenamer/Public/IAdvancedRenamerModule.h
@@ -0,0 +1,40 @@
+// Copyright Epic Games, Inc. All Rights Reserved.
+
+#pragma once
+
+#include "Containers/ContainersFwd.h"
+#include "Modules/ModuleManager.h"
+#include "Templates/SharedPointerFwd.h"
+
+class AActor;
+class IAdvancedRenamerProvider;
+class IToolkitHost;
+class SWidget;
+
+/** Advanced Rename Panel Plugin - Easily bulk rename stuff! */
+class IAdvancedRenamerModule : public IModuleInterface
+{
+protected:
+	static constexpr const TCHAR* ModuleName = TEXT("AdvancedRenamer");
+
+public:
+	static bool IsLoaded()
+	{
+		return FModuleManager::Get().IsModuleLoaded(ModuleName);
+	}
+
+	static IAdvancedRenamerModule& Get()
+	{
+		return FModuleManager::LoadModuleChecked<IAdvancedRenamerModule>(ModuleName);
+	}
+
+	virtual void OpenAdvancedRenamer(const TSharedRef<IAdvancedRenamerProvider>& InRenameProvider, const TSharedPtr<IToolkitHost>& InToolkitHost) = 0;
+
+	virtual void OpenAdvancedRenamer(const TSharedRef<IAdvancedRenamerProvider>& InRenameProvider, const TSharedPtr<SWidget>& InParentWidget) = 0;
+
+	virtual void OpenAdvancedRenamerForActors(const TArray<AActor*>& InActors, const TSharedPtr<IToolkitHost>& InToolkitHost) = 0;
+
+	virtual void OpenAdvancedRenamerForActors(const TArray<AActor*>& InActors, const TSharedPtr<SWidget>& InParentWidget) = 0;
+
+	virtual TArray<AActor*> GetActorsSharingClassesInWorld(const TArray<AActor*>& InActors) = 0;
+};
diff --git a/Engine/Plugins/Experimental/AdvancedRenamer/Source/AdvancedRenamer/Public/Providers/AdvancedRenamerActorProvider.h b/Engine/Plugins/Experimental/AdvancedRenamer/Source/AdvancedRenamer/Public/Providers/AdvancedRenamerActorProvider.h
deleted file mode 100644
diff --git a/Engine/Plugins/Experimental/AdvancedRenamer/Source/AdvancedRenamer/Public/Providers/AdvancedRenamerAssetProvider.h b/Engine/Plugins/Experimental/AdvancedRenamer/Source/AdvancedRenamer/Public/Providers/AdvancedRenamerAssetProvider.h
deleted file mode 100644
diff --git a/Engine/Plugins/Experimental/AdvancedRenamer/Source/AdvancedRenamer/Public/Providers/AdvancedRenamerObjectAssetProvider.h b/Engine/Plugins/Experimental/AdvancedRenamer/Source/AdvancedRenamer/Public/Providers/AdvancedRenamerObjectAssetProvider.h
deleted file mode 100644
diff --git a/Engine/Plugins/Experimental/AdvancedRenamer/Source/AdvancedRenamer/Public/Providers/AdvancedRenamerObjectObjectProvider.h b/Engine/Plugins/Experimental/AdvancedRenamer/Source/AdvancedRenamer/Public/Providers/AdvancedRenamerObjectObjectProvider.h
deleted file mode 100644
diff --git a/Engine/Plugins/Experimental/AdvancedRenamer/Source/AdvancedRenamer/Public/Providers/AdvancedRenamerObjectProvider.h b/Engine/Plugins/Experimental/AdvancedRenamer/Source/AdvancedRenamer/Public/Providers/AdvancedRenamerObjectProvider.h
deleted file mode 100644
diff --git a/Engine/Plugins/Experimental/AdvancedRenamer/Source/AdvancedRenamer/Public/Providers/AdvancedRenamerObjectProviderInterface.h b/Engine/Plugins/Experimental/AdvancedRenamer/Source/AdvancedRenamer/Public/Providers/AdvancedRenamerObjectProviderInterface.h
deleted file mode 100644
diff --git a/Engine/Plugins/Experimental/AdvancedRenamer/Source/AdvancedRenamer/Public/Providers/AdvancedRenamerProviderInterface.h b/Engine/Plugins/Experimental/AdvancedRenamer/Source/AdvancedRenamer/Public/Providers/AdvancedRenamerProviderInterface.h
deleted file mode 100644
diff --git a/Engine/Plugins/Experimental/AdvancedRenamer/Source/AdvancedRenamer/Public/Providers/IAdvancedRenamerProvider.h b/Engine/Plugins/Experimental/AdvancedRenamer/Source/AdvancedRenamer/Public/Providers/IAdvancedRenamerProvider.h
new file mode 100644
--- /dev/null
+++ b/Engine/Plugins/Experimental/AdvancedRenamer/Source/AdvancedRenamer/Public/Providers/IAdvancedRenamerProvider.h
@@ -0,0 +1,20 @@
+// Copyright Epic Games, Inc. All Rights Reserved.
+
+#pragma once
+
+#include "HAL/Platform.h"
+
+class FString;
+
+class IAdvancedRenamerProvider
+{
+public:
+	virtual ~IAdvancedRenamerProvider() = default;
+	virtual int32 Num() const = 0;
+	virtual bool IsValidIndex(int32 Index) const = 0;
+	virtual uint32 GetHash(int32 Index) const = 0;
+	virtual FString GetOriginalName(int32 Index) const = 0;
+	virtual bool RemoveIndex(int32 Index) = 0;
+	virtual bool CanRename(int32 Index) const = 0;
+	virtual bool ExecuteRename(int32 Index, const FString& NewName) = 0;
+};
diff --git a/Engine/Plugins/Experimental/Avalanche/Source/AvalancheEditor/Private/AdvancedRenamer/AvaAdvancedRenamerExtension.cpp b/Engine/Plugins/Experimental/Avalanche/Source/AvalancheEditor/Private/AdvancedRenamer/AvaAdvancedRenamerExtension.cpp
--- a/Engine/Plugins/Experimental/Avalanche/Source/AvalancheEditor/Private/AdvancedRenamer/AvaAdvancedRenamerExtension.cpp
+++ b/Engine/Plugins/Experimental/Avalanche/Source/AvalancheEditor/Private/AdvancedRenamer/AvaAdvancedRenamerExtension.cpp
@@ -1,12 +1,11 @@
 // Copyright Epic Games, Inc. All Rights Reserved.
 
 #include "AvaAdvancedRenamerExtension.h"
-#include "AdvancedRenamerBlueprintLibrary.h"
 #include "Algo/Transform.h"
 #include "AvaEditorCommands.h"
 #include "EditorModeManager.h"
 #include "EngineUtils.h"
-#include "Providers/AdvancedRenamerObjectProvider.h"
+#include "IAdvancedRenamerModule.h"
 #include "Selection.h"
 #include "Toolkits/IToolkitHost.h"
 #include "UObject/WeakObjectPtr.h"
@@ -66,7 +65,7 @@ bool FAvaAdvancedRenamerExtension::CanOpenAdvancedRenamerTool_ClassActors() cons
 void FAvaAdvancedRenamerExtension::OpenAdvancedRenamerTool_SharedClassActors() const
 {
 	TArray<AActor*> SelectedActors = GetSelectedActors();
-	SelectedActors = UAdvancedRenamerBlueprintLibrary::GetActorsSharingClassesInWorld(SelectedActors);
+	SelectedActors = IAdvancedRenamerModule::Get().GetActorsSharingClassesInWorld(SelectedActors);
 	
 	OpenAdvancedRenamerTool(SelectedActors);
 }
@@ -112,5 +111,5 @@ void FAvaAdvancedRenamerExtension::OpenAdvancedRenamerTool(const TArray<AActor*>
 		}
 	}
 
-	UAdvancedRenamerBlueprintLibrary::OpenAdvancedRenamerForActors(InActors, HostWindow);
+	IAdvancedRenamerModule::Get().OpenAdvancedRenamerForActors(InActors, HostWindow);
 }
