Advanced Renamer
- The renamer can now be automated - i.e. run without a UI.
- The renamer window is now resizeable.
- Small redesign to the UI to remove unnecessary clicks.

diff --git a/Engine/Plugins/Experimental/AdvancedRenamer/Source/AdvancedRenamer/Private/AdvancedRenamer.cpp b/Engine/Plugins/Experimental/AdvancedRenamer/Source/AdvancedRenamer/Private/AdvancedRenamer.cpp
new file mode 100644
--- /dev/null
+++ b/Engine/Plugins/Experimental/AdvancedRenamer/Source/AdvancedRenamer/Private/AdvancedRenamer.cpp
@@ -0,0 +1,489 @@
+// Copyright Epic Games, Inc. All Rights Reserved.
+
+#include "AdvancedRenamer.h"
+
+#include "AdvancedRenamerModule.h"
+#include "Internationalization/Regex.h"
+
+FAdvancedRenamer::FAdvancedRenamer(const TSharedRef<IAdvancedRenamerProvider>& InProvider, FAdvancedRenamerOptions* InInitialOptions)
+	: Provider(InProvider)
+{
+	if (InInitialOptions)
+	{
+		SetOption(*InInitialOptions);
+	}
+
+	int32 Count = Num();
+	check(Count > 0);
+
+	for (int32 Index = 0; Index < Count; ++Index)
+	{
+		if (!CanRename(Index))
+		{
+			RemoveIndex(Index);
+			--Index;
+			--Count;
+			continue;
+		}
+
+		const int32 Hash = GetHash(Index);
+		FString OriginalName = GetOriginalName(Index);
+
+		Previews.Add(MakeShared<FAdvancedRenamerPreview>(Hash, OriginalName));
+	}
+}
+
+const TSharedRef<IAdvancedRenamerProvider>& FAdvancedRenamer::GetProvider() const
+{
+	return Provider;
+}
+
+const FAdvancedRenamerOptions& FAdvancedRenamer::GetOptions() const
+{
+	return Options;
+}
+
+FAdvancedRenamerOptions& FAdvancedRenamer::GetOptions()
+{
+	return Options;
+}
+
+void FAdvancedRenamer::SetOption(const FAdvancedRenamerOptions& InOptions)
+{
+	Options = InOptions;
+	MarkDirty();
+}
+
+const TArray<TSharedPtr<FAdvancedRenamerPreview>>& FAdvancedRenamer::GetPreviews()
+{
+	return Previews;
+}
+
+TSharedPtr<FAdvancedRenamerPreview> FAdvancedRenamer::GetPreview(int32 InIndex) const
+{
+	if (Previews.IsValidIndex(InIndex))
+	{
+		return Previews[InIndex];
+	}
+
+	return nullptr;
+}
+
+bool FAdvancedRenamer::HasRenames() const
+{
+	return bHasRenames;
+}
+
+bool FAdvancedRenamer::IsDirty() const
+{
+	return bDirty;
+}
+
+void FAdvancedRenamer::MarkDirty()
+{
+	bDirty = true;
+}
+
+void FAdvancedRenamer::MarkClean()
+{
+	bDirty = false;
+}
+
+bool FAdvancedRenamer::UpdatePreviews()
+{
+	bHasRenames = false;
+	const int32 Count = Previews.Num();
+
+	for (int32 Index = 0; Index < Count; ++Index)
+	{
+		if (!Previews[Index].IsValid() || !IsValidIndex(Index))
+		{
+			RemoveIndex(Index);
+			--Index;
+			continue;
+		}
+
+		// Force recreation
+		Previews[Index]->NewName = ApplyRename(Previews[Index]->OriginalName, Index);
+
+		if (!Previews[Index]->NewName.IsEmpty())
+		{
+			continue;
+		}
+
+		if (GetOriginalName(Index) == Previews[Index]->NewName)
+		{
+			continue;
+		}
+
+		bHasRenames = true;
+	}
+
+	MarkClean();
+
+	return bHasRenames;
+}
+
+bool FAdvancedRenamer::Execute()
+{
+	if (!HasRenames())
+	{
+		UpdatePreviews();
+
+		if (!HasRenames())
+		{
+			return false;
+		}
+	}
+
+	const int32 Count = Previews.Num();
+	bool bAllSuccess = true;
+
+	for (int32 Index = 0; Index < Count; ++Index)
+	{
+		if (!Previews[Index].IsValid())
+		{
+			continue;
+		}
+
+		if (!IsValidIndex(Index))
+		{
+			continue;
+		}
+
+		if (!Previews[Index]->NewName.IsEmpty())
+		{
+			continue;
+		}
+
+		if (!ExecuteRename(Index, Previews[Index]->NewName))
+		{
+			bAllSuccess = false;
+		}
+	}
+
+	MarkClean();
+
+	return bAllSuccess;
+}
+
+FString FAdvancedRenamer::ApplyRename(const FString& InOriginalName, int32 InIndex) const
+{
+	FString NewName = InOriginalName;
+
+	if (Options.HasBaseName())
+	{
+		NewName = ApplyBaseName(NewName);
+	}
+
+	if (Options.HasPrefix())
+	{
+		NewName = ApplyPrefix(NewName);
+	}
+
+	if (Options.HasSuffix())
+	{
+		NewName = ApplySuffix(NewName, InIndex);
+	}
+
+	if (Options.HasSearchAndReplace())
+	{
+		switch (Options.SearchAndReplaceType)
+		{
+			default:
+				break;
+
+			case EAdvancedRenamerSeachAndReplaceType::PlainText:
+				NewName = ApplySearchPlainText(NewName);
+				break;
+
+			case EAdvancedRenamerSeachAndReplaceType::RegularExpression:
+				NewName = ApplySearchReplaceRegex(NewName);
+				break;
+		}
+	}
+
+	return NewName;
+}
+
+FString FAdvancedRenamer::ApplyBaseName(const FString& InOriginalName) const
+{
+	if (Options.BaseName.IsEmpty())
+	{
+		return InOriginalName;
+	}
+
+	return Options.BaseName;
+}
+
+FString FAdvancedRenamer::ApplyPrefix(const FString& InOriginalName) const
+{
+	FString Output = InOriginalName;
+
+	if (!Options.RemovePrefixSeparator.IsEmpty())
+	{
+		const int32 PrefixStart = Output.Find(Options.RemovePrefixSeparator, ESearchCase::IgnoreCase);
+
+		if (PrefixStart >= 0)
+		{
+			Output = Output.Mid(PrefixStart + Options.RemovePrefixSeparator.Len());
+		}
+	}
+	else if (Options.RemovePrefixCharacterCount > 0)
+	{
+		Output = Output.RightChop(Options.RemovePrefixCharacterCount);
+	}
+
+	if (!Options.AddPrefix.IsEmpty())
+	{
+		Output = Options.AddPrefix + Output;
+	}
+
+	return Output;
+}
+
+FString FAdvancedRenamer::ApplySuffix(const FString& InOriginalName, int32 InIndex) const
+{
+	static constexpr TCHAR FirstDigit = '0';
+	static constexpr TCHAR LastDigit = '9';
+
+	FString Output = InOriginalName;
+
+	if (!Options.RemoveSuffixSeparator.IsEmpty())
+	{
+		const int32 SuffixStart = Output.Find(Options.RemoveSuffixSeparator, ESearchCase::IgnoreCase, ESearchDir::FromEnd);
+
+		if (SuffixStart >= 0)
+		{
+			Output = Output.Mid(0, SuffixStart);
+		}
+	}
+	else if (Options.RemoveSuffixCharacterCount > 0)
+	{
+		Output = Output.LeftChop(Options.RemoveSuffixCharacterCount);
+	}
+
+	if (Options.bRemoveSuffixNumber || Options.bAddSuffixNumber)
+	{
+		int32 LastDigitIndex = Output.Len() - 1;
+
+		while (LastDigitIndex >= 0)
+		{
+			if (Output[LastDigitIndex] < FirstDigit || Output[LastDigitIndex] > LastDigit)
+			{
+				break;
+			}
+
+			--LastDigitIndex;
+		}
+
+		Output = Output.Mid(0, LastDigitIndex + 1);
+	}
+
+	if (!Options.AddSuffix.IsEmpty())
+	{
+		Output = Output + Options.AddSuffix;
+	}
+
+	if (Options.bAddSuffixNumber)
+	{
+		const int32 Current = Options.AddSuffixNumberStart + (InIndex * Options.AddSuffixNumberStep);
+		Output += FString::FromInt(Current);
+	}
+
+	return Output;
+}
+
+FString FAdvancedRenamer::ApplySearchPlainText(const FString& InOriginalName) const
+{
+	if (Options.SearchAndReplaceFromText.IsEmpty())
+	{
+		return InOriginalName;
+	}
+
+	return InOriginalName.Replace(
+		*Options.SearchAndReplaceFromText, 
+		*Options.SearchAndReplaceToText, 
+		Options.SearchAndReplaceCase
+	);
+}
+
+FString FAdvancedRenamer::ApplySearchReplaceRegex(const FString& InOriginalName) const
+{
+	if (Options.SearchAndReplaceFromText.IsEmpty())
+	{
+		return InOriginalName;
+	}
+
+	const FRegexPattern RegexPattern = FRegexPattern(
+		Options.SearchAndReplaceFromText,
+		Options.SearchAndReplaceCase == ESearchCase::IgnoreCase ? ERegexPatternFlags::CaseInsensitive : ERegexPatternFlags::None
+	);
+
+	return RegexReplace(InOriginalName, RegexPattern, Options.SearchAndReplaceToText);
+}
+
+FString FAdvancedRenamer::RegexReplace(const FString& InOriginalString, const FRegexPattern& InPattern, const FString& InReplaceString) const
+{
+	static const FString EscapeString = TEXT("\\");
+	static constexpr TCHAR EscapeChar = '\\';
+	static const FString GroupString = TEXT("$");
+	static constexpr TCHAR GroupChar = '$';
+	static constexpr TCHAR FirstDigit = '0';
+	static constexpr TCHAR LastDigit = '9';
+	static constexpr TCHAR NullChar = 0;
+
+	FRegexMatcher Matcher(InPattern, InOriginalString);
+
+	FString Output = "";
+	int32 StartCharIndex = 0;
+	bool bReadingGroupName = false;
+	int32 GroupIndex = INDEX_NONE;
+
+	while (Matcher.FindNext())
+	{
+		// Add on part of string after start/previous match
+		if (StartCharIndex != Matcher.GetMatchBeginning())
+		{
+			Output += InOriginalString.Mid(StartCharIndex, Matcher.GetMatchBeginning() - StartCharIndex);
+		}
+
+		bool bEscaped = false;
+
+		for (int32 CharIndex = 0; CharIndex <= InReplaceString.Len(); ++CharIndex)
+		{
+			const TCHAR& Char = CharIndex < InReplaceString.Len() ? InReplaceString[CharIndex] : NullChar;
+
+			if (bReadingGroupName)
+			{
+				// Build group index
+				if (Char >= FirstDigit && Char <= LastDigit)
+				{
+					if (GroupIndex == INDEX_NONE)
+					{
+						GroupIndex = 0;
+					}
+					else if (GroupIndex > 0)
+					{
+						GroupIndex *= 10;
+					}
+
+					const int32 NextDigit = static_cast<int32>(Char - FirstDigit);
+					GroupIndex += NextDigit;
+					continue;
+				}
+				// We've read a group index, add it to the output string
+				else if (GroupIndex > 0)
+				{
+					if (Matcher.GetCaptureGroupBeginning(GroupIndex) == INDEX_NONE)
+					{
+						UE_LOG(LogARP, Error, TEXT("Regex: Capture group does not exist %d."), GroupIndex);
+					}
+
+					Output += Matcher.GetCaptureGroup(GroupIndex);
+				}
+				// $0 matches the entire matched string
+				else if (GroupIndex == 0)
+				{
+					Output += InOriginalString.Mid(Matcher.GetMatchBeginning(), Matcher.GetMatchEnding() - Matcher.GetMatchBeginning());
+				}
+				// An unescaped $
+				else
+				{
+					UE_LOG(LogARP, Error, TEXT("Regex: Unescaped %s."), *GroupString);
+
+					Output += GroupString;
+				}
+
+				bReadingGroupName = false;
+				// Continue regular parsing of this character.
+			}
+
+			// Check for special chars
+			if (!bEscaped)
+			{
+				if (Char == EscapeChar)
+				{
+					bEscaped = true;
+					continue;
+				}
+
+				if (Char == GroupChar)
+				{
+					bReadingGroupName = true;
+					GroupIndex = INDEX_NONE;
+					continue;
+				}
+			}
+			else
+			{
+				// If the last char is a \ assume that it's not an escape char
+				if (Char == NullChar)
+				{
+					UE_LOG(LogARP, Error, TEXT("Regex: Unescaped %s."), *EscapeString);
+
+					Output += EscapeChar;
+				}
+			}
+
+			if (Char != NullChar)
+			{
+				Output += Char;
+			}
+
+			bEscaped = false;
+		}
+
+		StartCharIndex = Matcher.GetMatchEnding();
+	}
+
+	// Add on the end of the string after the last match
+	if (StartCharIndex < InOriginalString.Len())
+	{
+		Output += InOriginalString.Mid(StartCharIndex);
+	}
+
+	return Output;
+}
+
+int32 FAdvancedRenamer::Num() const
+{
+	return Provider->Num();
+}
+
+bool FAdvancedRenamer::IsValidIndex(int32 InIndex) const
+{
+	return Provider->IsValidIndex(InIndex);
+}
+
+uint32 FAdvancedRenamer::GetHash(int32 InIndex) const
+{
+	return Provider->GetHash(InIndex);
+}
+
+FString FAdvancedRenamer::GetOriginalName(int32 InIndex) const
+{
+	return Provider->GetOriginalName(InIndex);
+}
+
+bool FAdvancedRenamer::RemoveIndex(int32 InIndex)
+{
+	// Can fail during construction when indices that aren't renameable are removed from the provider before
+	// they are added to ListData.
+	if (Previews.IsValidIndex(InIndex))
+	{
+		Previews.RemoveAt(InIndex);
+	}
+
+	return Provider->RemoveIndex(InIndex);
+}
+
+bool FAdvancedRenamer::CanRename(int32 InIndex) const
+{
+	return Provider->CanRename(InIndex);
+}
+
+bool FAdvancedRenamer::ExecuteRename(int32 InIndex, const FString& InNewName)
+{
+	return Provider->ExecuteRename(InIndex, InNewName);
+}
diff --git a/Engine/Plugins/Experimental/AdvancedRenamer/Source/AdvancedRenamer/Private/AdvancedRenamer.h b/Engine/Plugins/Experimental/AdvancedRenamer/Source/AdvancedRenamer/Private/AdvancedRenamer.h
new file mode 100644
--- /dev/null
+++ b/Engine/Plugins/Experimental/AdvancedRenamer/Source/AdvancedRenamer/Private/AdvancedRenamer.h
@@ -0,0 +1,58 @@
+// Copyright Epic Games, Inc. All Rights Reserved.
+
+#pragma once
+
+#include "IAdvancedRenamer.h"
+#include "Providers/IAdvancedRenamerProvider.h"
+#include "Templates/SharedPointer.h"
+
+class FRegexPattern;
+struct FAdvancedRenamerOptions;
+struct FAdvancedRenamerPreview;
+
+class FAdvancedRenamer : public IAdvancedRenamer
+{
+public:
+	FAdvancedRenamer(const TSharedRef<IAdvancedRenamerProvider>& InProvider, FAdvancedRenamerOptions* InInitialOptions = nullptr);
+
+	//~ Begin IAdvancedRenamer
+	virtual const TSharedRef<IAdvancedRenamerProvider>& GetProvider() const override;
+	virtual const FAdvancedRenamerOptions& GetOptions() const override;
+	virtual FAdvancedRenamerOptions& GetOptions() override;
+	virtual void SetOption(const FAdvancedRenamerOptions& InOptions) override;
+	virtual const TArray<TSharedPtr<FAdvancedRenamerPreview>>& GetPreviews() override;
+	virtual TSharedPtr<FAdvancedRenamerPreview> GetPreview(int32 InIndex) const override;
+	virtual bool HasRenames() const override;
+	virtual bool IsDirty() const override;
+	virtual void MarkDirty() override;
+	virtual void MarkClean() override;
+	virtual FString ApplyRename(const FString& InName, int32 InIndex) const override;
+	virtual bool UpdatePreviews() override;
+	virtual bool Execute() override;
+	//~ End IAdvancedRenamer
+
+	//~ Begin IAdvancedRenamerProvider
+	virtual int32 Num() const override;
+	virtual bool IsValidIndex(int32 InIndex) const override;
+	virtual uint32 GetHash(int32 InIndex) const override;
+	virtual FString GetOriginalName(int32 InIndex) const override;
+	virtual bool RemoveIndex(int32 InIndex) override;
+	virtual bool CanRename(int32 InIndex) const override;
+	virtual bool ExecuteRename(int32 InIndex, const FString& InNewName) override;
+	//~ End IAdvancedRenamerProvider
+
+protected:
+	TSharedRef<IAdvancedRenamerProvider> Provider;
+	TArray<TSharedPtr<FAdvancedRenamerPreview>> Previews;
+	FAdvancedRenamerOptions Options;
+	bool bHasRenames;
+	bool bDirty;
+
+	FString ApplyBaseName(const FString& InOriginalName) const;
+	FString ApplyPrefix(const FString& InOriginalName) const;
+	FString ApplySuffix(const FString& InOriginalName, int32 InIndex) const;
+	FString ApplySearchPlainText(const FString& InOriginalName) const;
+	FString ApplySearchReplaceRegex(const FString& InOriginalName) const;
+
+	FString RegexReplace(const FString& InOriginalString, const FRegexPattern& InPattern, const FString& InReplaceString) const;
+};
diff --git a/Engine/Plugins/Experimental/AdvancedRenamer/Source/AdvancedRenamer/Private/AdvancedRenamerModule.cpp b/Engine/Plugins/Experimental/AdvancedRenamer/Source/AdvancedRenamer/Private/AdvancedRenamerModule.cpp
--- a/Engine/Plugins/Experimental/AdvancedRenamer/Source/AdvancedRenamer/Private/AdvancedRenamerModule.cpp
+++ b/Engine/Plugins/Experimental/AdvancedRenamer/Source/AdvancedRenamer/Private/AdvancedRenamerModule.cpp
@@ -1,6 +1,8 @@
 // Copyright Epic Games, Inc. All Rights Reserved.
 
 #include "AdvancedRenamerModule.h"
+
+#include "AdvancedRenamer.h"
 #include "AdvancedRenamerCommands.h"
 #include "AdvancedRenamerStyle.h"
 #include "EngineUtils.h"
@@ -22,10 +24,13 @@ namespace UE::AdvancedRenamer::Private
 	{
 		return SNew(SWindow)
 			.Title(LOCTEXT("AdvancedRenameWindow", "Rename Actors"))
-			.ClientSize(FVector2D(600.0f, 500.0f))
-			.SizingRule(ESizingRule::FixedSize)
+			.ClientSize(FVector2D(600, 538))
+			.SizingRule(ESizingRule::UserSized)
 			.SupportsMaximize(false)
-			.SupportsMinimize(false);
+			.SupportsMinimize(false)
+			.MinWidth(600)
+			.MinHeight(538)
+			.MaxHeight(538);
 	}
 }
 
@@ -45,21 +50,19 @@ void FAdvancedRenamerModule::ShutdownModule()
 	FAdvancedRenamerLevelEditorIntegration::Shutdown();
 }
 
-void FAdvancedRenamerModule::OpenAdvancedRenamer(const TSharedRef<IAdvancedRenamerProvider>& InRenameProvider, const TSharedPtr<SWidget>& InParentWidget)
+TSharedRef<IAdvancedRenamer> FAdvancedRenamerModule::CreateAdvancedRenamer(const TSharedRef<IAdvancedRenamerProvider>& InRenameProvider)
 {
-	TSharedRef<SWindow> AdvancedRenameWindow = UE::AdvancedRenamer::Private::CreateAdvancedRenamerWindow();
-	AdvancedRenameWindow->SetContent(SNew(SAdvancedRenamerPanel).SharedProvider(InRenameProvider));
+	return MakeShared<FAdvancedRenamer>(InRenameProvider);
+}
 
-	TSharedPtr<SWidget> ParentWindow = FSlateApplication::Get().FindBestParentWindowForDialogs(InParentWidget);
-	FSlateApplication::Get().AddModalWindow(AdvancedRenameWindow, ParentWindow);
+void FAdvancedRenamerModule::OpenAdvancedRenamer(const TSharedRef<IAdvancedRenamerProvider>& InRenameProvider, const TSharedPtr<SWidget>& InParentWidget)
+{
+	return OpenAdvancedRenamer(CreateAdvancedRenamer(InRenameProvider), InParentWidget);
 }
 
 void FAdvancedRenamerModule::OpenAdvancedRenamer(const TSharedRef<IAdvancedRenamerProvider>& InRenameProvider, const TSharedPtr<IToolkitHost>& InToolkitHost)
 {
-	if (InToolkitHost.IsValid())
-	{
-		OpenAdvancedRenamer(InRenameProvider, InToolkitHost->GetParentWidget());
-	}
+	return OpenAdvancedRenamer(CreateAdvancedRenamer(InRenameProvider), InToolkitHost);
 }
 
 void FAdvancedRenamerModule::OpenAdvancedRenamerForActors(const TArray<AActor*>& InActors, const TSharedPtr<SWidget>& InParentWidget)
@@ -91,6 +94,23 @@ void FAdvancedRenamerModule::OpenAdvancedRenamerForActors(const TArray<AActor*>&
 	}
 }
 
+void FAdvancedRenamerModule::OpenAdvancedRenamer(const TSharedRef<IAdvancedRenamer>& InRenamer, const TSharedPtr<IToolkitHost>& InToolkitHost)
+{
+	if (InToolkitHost.IsValid())
+	{
+		OpenAdvancedRenamer(InRenamer, InToolkitHost->GetParentWidget());
+	}
+}
+
+void FAdvancedRenamerModule::OpenAdvancedRenamer(const TSharedRef<IAdvancedRenamer>& InRenamer, const TSharedPtr<SWidget>& InParentWidget)
+{
+	TSharedRef<SWindow> AdvancedRenameWindow = UE::AdvancedRenamer::Private::CreateAdvancedRenamerWindow();
+	AdvancedRenameWindow->SetContent(SNew(SAdvancedRenamerPanel, InRenamer));
+
+	TSharedPtr<SWidget> ParentWindow = FSlateApplication::Get().FindBestParentWindowForDialogs(InParentWidget);
+	FSlateApplication::Get().AddModalWindow(AdvancedRenameWindow, ParentWindow);
+}
+
 TArray<AActor*> FAdvancedRenamerModule::GetActorsSharingClassesInWorld(const TArray<AActor*>& InActors)
 {
 	TSet<UClass*> SelectedClasses;
diff --git a/Engine/Plugins/Experimental/AdvancedRenamer/Source/AdvancedRenamer/Private/AdvancedRenamerModule.h b/Engine/Plugins/Experimental/AdvancedRenamer/Source/AdvancedRenamer/Private/AdvancedRenamerModule.h
--- a/Engine/Plugins/Experimental/AdvancedRenamer/Source/AdvancedRenamer/Private/AdvancedRenamerModule.h
+++ b/Engine/Plugins/Experimental/AdvancedRenamer/Source/AdvancedRenamer/Private/AdvancedRenamerModule.h
@@ -16,10 +16,13 @@ class FAdvancedRenamerModule : public IAdvancedRenamerModule
 	//~ End IModuleInterface
 
 	//~ Begin IAdvancedRenamerModule
+	virtual TSharedRef<IAdvancedRenamer> CreateAdvancedRenamer(const TSharedRef<IAdvancedRenamerProvider>& InRenameProvider) override;
 	virtual void OpenAdvancedRenamer(const TSharedRef<IAdvancedRenamerProvider>& InRenameProvider, const TSharedPtr<SWidget>& InParentWidget) override;
 	virtual void OpenAdvancedRenamer(const TSharedRef<IAdvancedRenamerProvider>& InRenameProvider, const TSharedPtr<IToolkitHost>& InToolkitHost) override;
 	virtual void OpenAdvancedRenamerForActors(const TArray<AActor*>& InActors, const TSharedPtr<SWidget>& InParentWidget) override;
 	virtual void OpenAdvancedRenamerForActors(const TArray<AActor*>& InActors, const TSharedPtr<IToolkitHost>& InToolkitHost) override;
+	virtual void OpenAdvancedRenamer(const TSharedRef<IAdvancedRenamer>& InRenamer, const TSharedPtr<IToolkitHost>& InToolkitHost) override;
+	virtual void OpenAdvancedRenamer(const TSharedRef<IAdvancedRenamer>& InRenamer, const TSharedPtr<SWidget>& InParentWidget) override;
 	virtual TArray<AActor*> GetActorsSharingClassesInWorld(const TArray<AActor*>& InActors) override;
 	//~ End IAdvancedRenamerModule
 };
diff --git a/Engine/Plugins/Experimental/AdvancedRenamer/Source/AdvancedRenamer/Private/IAdvancedRenamer.h b/Engine/Plugins/Experimental/AdvancedRenamer/Source/AdvancedRenamer/Private/IAdvancedRenamer.h
new file mode 100644
--- /dev/null
+++ b/Engine/Plugins/Experimental/AdvancedRenamer/Source/AdvancedRenamer/Private/IAdvancedRenamer.h
@@ -0,0 +1,151 @@
+// Copyright Epic Games, Inc. All Rights Reserved.
+
+#pragma once
+
+#include "Containers/ContainersFwd.h"
+#include "Containers/UnrealString.h"
+#include "HAL/Platform.h"
+#include "Misc/CString.h"
+#include "Providers/IAdvancedRenamerProvider.h"
+#include "Templates/SharedPointerFwd.h"
+
+enum class EAdvancedRenamerSeachAndReplaceType : uint8
+{
+	None,
+	PlainText,
+	RegularExpression
+};
+
+/**
+ * Options for the rename engine.
+ *
+ * - RemovePrefixSeparator and RemovePrefixCharacterCount are mutually exclusive.
+ * - RemoveSuffixSeparator and RemoveSuffixCharacterCount are mutually exclusive.
+ * - bAddSuffixNumber implies bRemoveSuffixNumer.
+ */
+struct FAdvancedRenamerOptions
+{
+	/** Name => OtherName */
+	FString BaseName = ""; 
+
+	/** Name => Other_Name */
+	FString AddPrefix = ""; 
+
+	/** Other_Name, _ => Name */
+	FString RemovePrefixSeparator = "";
+
+	/** OtherName, 5 => Name */
+	int32 RemovePrefixCharacterCount = 0; 
+
+	/** Name => Name_Other */
+	FString AddSuffix = ""; 
+
+	/** Name => Name1 */
+	bool bAddSuffixNumber = false; 
+
+	/** Name, 5 => Name5 */
+	int32 AddSuffixNumberStart = 1; 
+
+	/** Name, Moo, 1, 5 => Name1, Moo6 */
+	int32 AddSuffixNumberStep = 1;
+
+	/** Name_Other, _ => Name */
+	FString RemoveSuffixSeparator = ""; 
+
+	/** NameOther, 5 => Name */
+	int32 RemoveSuffixCharacterCount = 0; 
+
+	/** Name5 => Name */
+	bool bRemoveSuffixNumber = false; 
+
+	EAdvancedRenamerSeachAndReplaceType SearchAndReplaceType = EAdvancedRenamerSeachAndReplaceType::None;
+	ESearchCase::Type SearchAndReplaceCase = ESearchCase::CaseSensitive;
+
+	/** Name, am, ot => Note */
+	FString SearchAndReplaceFromText = "";
+
+	/** Name, am, ot => Note */
+	FString SearchAndReplaceToText = "";
+
+	bool HasBaseName() const
+	{
+		return !BaseName.IsEmpty();
+	}
+
+	bool HasPrefix() const
+	{
+		return !AddPrefix.IsEmpty()
+			|| !RemovePrefixSeparator.IsEmpty()
+			|| RemovePrefixCharacterCount > 0;
+	}
+
+	bool HasSuffix() const
+	{
+		return !AddSuffix.IsEmpty()
+			|| (bAddSuffixNumber && AddSuffixNumberStep > 0)
+			|| !RemoveSuffixSeparator.IsEmpty()
+			|| RemoveSuffixCharacterCount > 0
+			|| bRemoveSuffixNumber;			
+	}
+
+	bool HasSearchAndReplace() const
+	{
+		return SearchAndReplaceType != EAdvancedRenamerSeachAndReplaceType::None
+			&& !SearchAndReplaceFromText.IsEmpty();
+	}
+};
+
+struct FAdvancedRenamerPreview
+{
+	FAdvancedRenamerPreview(int32 InHash, const FString InOriginalName)
+		: Hash(InHash)
+		, OriginalName(InOriginalName)
+		, NewName(FString(""))
+	{
+	}
+
+	int32 Hash;
+	const FString OriginalName;
+	FString NewName;
+};
+
+/**
+* Implements its own provider interface so it can avoid long Execute_ lines and handle
+* the 2 different types of provider (SharedPtr and UObject.)
+*/
+class IAdvancedRenamer : public IAdvancedRenamerProvider
+{
+public:
+	virtual ~IAdvancedRenamer() = default;
+
+	virtual const TSharedRef<IAdvancedRenamerProvider>& GetProvider() const = 0;
+
+	virtual const FAdvancedRenamerOptions& GetOptions() const = 0;
+
+	/** Make sure to run UpdatePreviews or MarkDirty if you make changes directly. */
+	virtual FAdvancedRenamerOptions& GetOptions() = 0;
+
+	virtual void SetOption(const FAdvancedRenamerOptions& InOptions) = 0;
+
+	virtual const TArray<TSharedPtr<FAdvancedRenamerPreview>>& GetPreviews() = 0;
+
+	/** Returns the preview for the item at the given index. */
+	virtual TSharedPtr<FAdvancedRenamerPreview> GetPreview(int32 InIndex) const = 0;
+
+	/** True if there are any items actually renamed by the preview generator. */
+	virtual bool HasRenames() const = 0;
+
+	/** Whether the options have been updated. */
+	virtual bool IsDirty() const = 0;
+	virtual void MarkDirty() = 0;
+	virtual void MarkClean() = 0;
+
+	/** Executes the rename on the given name. */
+	virtual FString ApplyRename(const FString& InName, int32 InIndex) const = 0;
+
+	/** Returns true if any names actually changed. */
+	virtual bool UpdatePreviews() = 0;
+
+	/** Returns true if all items were updated without error. */
+	virtual bool Execute() = 0;
+};
diff --git a/Engine/Plugins/Experimental/AdvancedRenamer/Source/AdvancedRenamer/Private/Providers/AdvancedRenamerActorProvider.h b/Engine/Plugins/Experimental/AdvancedRenamer/Source/AdvancedRenamer/Private/Providers/AdvancedRenamerActorProvider.h
--- a/Engine/Plugins/Experimental/AdvancedRenamer/Source/AdvancedRenamer/Private/Providers/AdvancedRenamerActorProvider.h
+++ b/Engine/Plugins/Experimental/AdvancedRenamer/Source/AdvancedRenamer/Private/Providers/AdvancedRenamerActorProvider.h
@@ -25,7 +25,7 @@ protected:
 	//~ Begin IAdvancedRenamerProvider
 	virtual int32 Num() const override;
 	virtual bool IsValidIndex(int32 InIndex) const override;
-	virtual uint32 GetHash(int32 InIndex) const override;;
+	virtual uint32 GetHash(int32 InIndex) const override;
 	virtual FString GetOriginalName(int32 InIndex) const override;
 	virtual bool RemoveIndex(int32 InIndex) override;
 	virtual bool CanRename(int32 InIndex) const override;
diff --git a/Engine/Plugins/Experimental/AdvancedRenamer/Source/AdvancedRenamer/Private/Providers/AdvancedRenamerAssetProvider.cpp b/Engine/Plugins/Experimental/AdvancedRenamer/Source/AdvancedRenamer/Private/Providers/AdvancedRenamerAssetProvider.cpp
--- a/Engine/Plugins/Experimental/AdvancedRenamer/Source/AdvancedRenamer/Private/Providers/AdvancedRenamerAssetProvider.cpp
+++ b/Engine/Plugins/Experimental/AdvancedRenamer/Source/AdvancedRenamer/Private/Providers/AdvancedRenamerAssetProvider.cpp
@@ -30,14 +30,14 @@ void FAdvancedRenamerAssetProvider::AddAssetData(const FAssetData& InAsset)
 	AssetList.Add(InAsset);
 }
 
-UObject* FAdvancedRenamerAssetProvider::GetAsset(int32 Index) const
+UObject* FAdvancedRenamerAssetProvider::GetAsset(int32 InIndex) const
 {
-	if (!AssetList.IsValidIndex(Index))
+	if (!AssetList.IsValidIndex(InIndex))
 	{
 		return nullptr;
 	}
 
-	return AssetList[Index].GetAsset();
+	return AssetList[InIndex].GetAsset();
 }
 
 int32 FAdvancedRenamerAssetProvider::Num() const
@@ -45,16 +45,16 @@ int32 FAdvancedRenamerAssetProvider::Num() const
 	return AssetList.Num();
 }
 
-bool FAdvancedRenamerAssetProvider::IsValidIndex(int32 Index) const
+bool FAdvancedRenamerAssetProvider::IsValidIndex(int32 InIndex) const
 {
-	UObject* Asset = GetAsset(Index);
+	UObject* Asset = GetAsset(InIndex);
 
 	return IsValid(Asset) && Asset->IsAsset();
 }
 
-FString FAdvancedRenamerAssetProvider::GetOriginalName(int32 Index) const
+FString FAdvancedRenamerAssetProvider::GetOriginalName(int32 InIndex) const
 {
-	UObject* Asset = GetAsset(Index);
+	UObject* Asset = GetAsset(InIndex);
 
 	if (!IsValid(Asset))
 	{
@@ -64,9 +64,9 @@ FString FAdvancedRenamerAssetProvider::GetOriginalName(int32 Index) const
 	return Asset->GetName();
 }
 
-uint32 FAdvancedRenamerAssetProvider::GetHash(int32 Index) const
+uint32 FAdvancedRenamerAssetProvider::GetHash(int32 InIndex) const
 {
-	UObject* Asset = GetAsset(Index);
+	UObject* Asset = GetAsset(InIndex);
 
 	if (!IsValid(Asset))
 	{
@@ -76,20 +76,20 @@ uint32 FAdvancedRenamerAssetProvider::GetHash(int32 Index) const
 	return GetTypeHash(Asset);
 }
 
-bool FAdvancedRenamerAssetProvider::RemoveIndex(int32 Index)
+bool FAdvancedRenamerAssetProvider::RemoveIndex(int32 InIndex)
 {
-	if (!AssetList.IsValidIndex(Index))
+	if (!AssetList.IsValidIndex(InIndex))
 	{
 		return false;
 	}
 
-	AssetList.RemoveAt(Index);
+	AssetList.RemoveAt(InIndex);
 	return true;
 }
 
-bool FAdvancedRenamerAssetProvider::CanRename(int32 Index) const
+bool FAdvancedRenamerAssetProvider::CanRename(int32 InIndex) const
 {
-	UObject* Asset = GetAsset(Index);
+	UObject* Asset = GetAsset(InIndex);
 
 	if (!IsValid(Asset))
 	{
@@ -99,9 +99,9 @@ bool FAdvancedRenamerAssetProvider::CanRename(int32 Index) const
 	return true;
 }
 
-bool FAdvancedRenamerAssetProvider::ExecuteRename(int32 Index, const FString& NewName)
+bool FAdvancedRenamerAssetProvider::ExecuteRename(int32 InIndex, const FString& InNewName)
 {
-	UObject* Asset = GetAsset(Index);
+	UObject* Asset = GetAsset(InIndex);
 
 	if (!IsValid(Asset))
 	{
@@ -112,7 +112,7 @@ bool FAdvancedRenamerAssetProvider::ExecuteRename(int32 Index, const FString& Ne
 	PackagePath = FPaths::GetPath(PackagePath);
 
 	IAssetTools& AssetTools = FModuleManager::GetModuleChecked<FAssetToolsModule>("AssetTools").Get();
-	return AssetTools.RenameAssets({FAssetRenameData(Asset, PackagePath, NewName)});
+	return AssetTools.RenameAssets({FAssetRenameData(Asset, PackagePath, InNewName)});
 }
 
 #undef LOCTEXT_NAMESPACE
diff --git a/Engine/Plugins/Experimental/AdvancedRenamer/Source/AdvancedRenamer/Private/Providers/AdvancedRenamerAssetProvider.h b/Engine/Plugins/Experimental/AdvancedRenamer/Source/AdvancedRenamer/Private/Providers/AdvancedRenamerAssetProvider.h
--- a/Engine/Plugins/Experimental/AdvancedRenamer/Source/AdvancedRenamer/Private/Providers/AdvancedRenamerAssetProvider.h
+++ b/Engine/Plugins/Experimental/AdvancedRenamer/Source/AdvancedRenamer/Private/Providers/AdvancedRenamerAssetProvider.h
@@ -17,17 +17,17 @@ public:
 	void SetAssetList(const TArray<FAssetData>& InAssetList);
 	void AddAssetList(const TArray<FAssetData>& InAssetList);
 	void AddAssetData(const FAssetData& InAsset);
-	UObject* GetAsset(int32 Index) const;
+	UObject* GetAsset(int32 InIndex) const;
 
 protected:
 	//~ Begin IAdvancedRenamerProvider
 	virtual int32 Num() const override;
-	virtual bool IsValidIndex(int32 Index) const override;
-	virtual uint32 GetHash(int32 Index) const override;;
-	virtual FString GetOriginalName(int32 Index) const override;
-	virtual bool RemoveIndex(int32 Index) override;
-	virtual bool CanRename(int32 Index) const override;
-	virtual bool ExecuteRename(int32 Index, const FString& NewName) override;
+	virtual bool IsValidIndex(int32 InIndex) const override;
+	virtual uint32 GetHash(int32 InIndex) const override;;
+	virtual FString GetOriginalName(int32 InIndex) const override;
+	virtual bool RemoveIndex(int32 InIndex) override;
+	virtual bool CanRename(int32 InIndex) const override;
+	virtual bool ExecuteRename(int32 InIndex, const FString& InNewName) override;
 	//~ End IAdvancedRenamerProvider
 
 	TArray<FAssetData> AssetList;
diff --git a/Engine/Plugins/Experimental/AdvancedRenamer/Source/AdvancedRenamer/Private/Providers/AdvancedRenamerObjectProvider.cpp b/Engine/Plugins/Experimental/AdvancedRenamer/Source/AdvancedRenamer/Private/Providers/AdvancedRenamerObjectProvider.cpp
--- a/Engine/Plugins/Experimental/AdvancedRenamer/Source/AdvancedRenamer/Private/Providers/AdvancedRenamerObjectProvider.cpp
+++ b/Engine/Plugins/Experimental/AdvancedRenamer/Source/AdvancedRenamer/Private/Providers/AdvancedRenamerObjectProvider.cpp
@@ -27,14 +27,14 @@ void FAdvancedRenamerObjectProvider::AddObjectData(UObject* InObject)
 	ObjectList.Add(InObject);
 }
 
-UObject* FAdvancedRenamerObjectProvider::GetObject(int32 Index) const
+UObject* FAdvancedRenamerObjectProvider::GetObject(int32 InIndex) const
 {
-	if (!ObjectList.IsValidIndex(Index))
+	if (!ObjectList.IsValidIndex(InIndex))
 	{
 		return nullptr;
 	}
 
-	return ObjectList[Index].Get();
+	return ObjectList[InIndex].Get();
 }
 
 int32 FAdvancedRenamerObjectProvider::Num() const
@@ -42,16 +42,16 @@ int32 FAdvancedRenamerObjectProvider::Num() const
 	return ObjectList.Num();
 }
 
-bool FAdvancedRenamerObjectProvider::IsValidIndex(int32 Index) const
+bool FAdvancedRenamerObjectProvider::IsValidIndex(int32 InIndex) const
 {
-	UObject* Object = GetObject(Index);
+	UObject* Object = GetObject(InIndex);
 
 	return IsValid(Object);
 }
 
-FString FAdvancedRenamerObjectProvider::GetOriginalName(int32 Index) const
+FString FAdvancedRenamerObjectProvider::GetOriginalName(int32 InIndex) const
 {
-	UObject* Object = GetObject(Index);
+	UObject* Object = GetObject(InIndex);
 
 	if (!IsValid(Object))
 	{
@@ -61,9 +61,9 @@ FString FAdvancedRenamerObjectProvider::GetOriginalName(int32 Index) const
 	return Object->GetName();
 }
 
-uint32 FAdvancedRenamerObjectProvider::GetHash(int32 Index) const
+uint32 FAdvancedRenamerObjectProvider::GetHash(int32 InIndex) const
 {
-	UObject* Object = GetObject(Index);
+	UObject* Object = GetObject(InIndex);
 
 	if (!IsValid(Object))
 	{
@@ -73,20 +73,20 @@ uint32 FAdvancedRenamerObjectProvider::GetHash(int32 Index) const
 	return GetTypeHash(Object);
 }
 
-bool FAdvancedRenamerObjectProvider::RemoveIndex(int32 Index)
+bool FAdvancedRenamerObjectProvider::RemoveIndex(int32 InIndex)
 {
-	if (!ObjectList.IsValidIndex(Index))
+	if (!ObjectList.IsValidIndex(InIndex))
 	{
 		return false;
 	}
 
-	ObjectList.RemoveAt(Index);
+	ObjectList.RemoveAt(InIndex);
 	return true;
 }
 
-bool FAdvancedRenamerObjectProvider::CanRename(int32 Index) const
+bool FAdvancedRenamerObjectProvider::CanRename(int32 InIndex) const
 {
-	UObject* Object = GetObject(Index);
+	UObject* Object = GetObject(InIndex);
 
 	if (!IsValid(Object))
 	{
@@ -96,18 +96,18 @@ bool FAdvancedRenamerObjectProvider::CanRename(int32 Index) const
 	return true;
 }
 
-bool FAdvancedRenamerObjectProvider::ExecuteRename(int32 Index, const FString& NewName)
+bool FAdvancedRenamerObjectProvider::ExecuteRename(int32 InIndex, const FString& InNewName)
 {
-	UObject* Object = GetObject(Index);
+	UObject* Object = GetObject(InIndex);
 
 	if (!IsValid(Object))
 	{
 		return false;
 	}
 
-	if (Object->Rename(*NewName, nullptr, REN_Test))
+	if (Object->Rename(*InNewName, nullptr, REN_Test))
 	{
-		Object->Rename(*NewName);
+		Object->Rename(*InNewName);
 		return true;
 	}
 
diff --git a/Engine/Plugins/Experimental/AdvancedRenamer/Source/AdvancedRenamer/Private/Providers/AdvancedRenamerObjectProvider.h b/Engine/Plugins/Experimental/AdvancedRenamer/Source/AdvancedRenamer/Private/Providers/AdvancedRenamerObjectProvider.h
--- a/Engine/Plugins/Experimental/AdvancedRenamer/Source/AdvancedRenamer/Private/Providers/AdvancedRenamerObjectProvider.h
+++ b/Engine/Plugins/Experimental/AdvancedRenamer/Source/AdvancedRenamer/Private/Providers/AdvancedRenamerObjectProvider.h
@@ -19,17 +19,17 @@ public:
 	void SetObjectList(const TArray<TWeakObjectPtr<UObject>>& InObjectList);
 	void AddObjectList(const TArray<TWeakObjectPtr<UObject>>& InObjectList);
 	void AddObjectData(UObject* InObject);
-	UObject* GetObject(int32 Index) const;
+	UObject* GetObject(int32 InIndex) const;
 
 protected:
 	//~ Begin IAdvancedRenamerProvider
 	virtual int32 Num() const override;
-	virtual bool IsValidIndex(int32 Index) const override;
-	virtual uint32 GetHash(int32 Index) const override;;
-	virtual FString GetOriginalName(int32 Index) const override;
-	virtual bool RemoveIndex(int32 Index) override;
-	virtual bool CanRename(int32 Index) const override;
-	virtual bool ExecuteRename(int32 Index, const FString& NewName) override;
+	virtual bool IsValidIndex(int32 InIndex) const override;
+	virtual uint32 GetHash(int32 InIndex) const override;;
+	virtual FString GetOriginalName(int32 InIndex) const override;
+	virtual bool RemoveIndex(int32 InIndex) override;
+	virtual bool CanRename(int32 InIndex) const override;
+	virtual bool ExecuteRename(int32 InIndex, const FString& InNewName) override;
 	//~ End IAdvancedRenamerProvider
 
 	TArray<TWeakObjectPtr<UObject>> ObjectList;
diff --git a/Engine/Plugins/Experimental/AdvancedRenamer/Source/AdvancedRenamer/Private/Slate/SAdvancedRenamerPanel.cpp b/Engine/Plugins/Experimental/AdvancedRenamer/Source/AdvancedRenamer/Private/Slate/SAdvancedRenamerPanel.cpp
--- a/Engine/Plugins/Experimental/AdvancedRenamer/Source/AdvancedRenamer/Private/Slate/SAdvancedRenamerPanel.cpp
+++ b/Engine/Plugins/Experimental/AdvancedRenamer/Source/AdvancedRenamer/Private/Slate/SAdvancedRenamerPanel.cpp
@@ -1,13 +1,11 @@
 // Copyright Epic Games, Inc. All Rights Reserved.
 
 #include "Slate/SAdvancedRenamerPanel.h"
-#include "AdvancedRenamerModule.h"
 #include "AdvancedRenamerStyle.h"
 #include "EngineAnalytics.h"
 #include "Framework/Commands/GenericCommands.h"
 #include "Framework/MultiBox/MultiBoxBuilder.h"
-#include "Internationalization/Regex.h"
-#include "Providers/IAdvancedRenamerProvider.h"
+#include "IAdvancedRenamer.h"
 #include "ScopedTransaction.h"
 #include "Styling/AppStyle.h"
 #include "Styling/StyleColors.h"
@@ -16,9 +14,9 @@
 #include "Widgets/Input/SCheckBox.h"
 #include "Widgets/Input/SEditableTextBox.h"
 #include "Widgets/Input/SMultiLineEditableTextBox.h"
+#include "Widgets/Input/SSpinBox.h"
 #include "Widgets/Layout/SBorder.h"
 #include "Widgets/Layout/SBox.h"
-#include "Widgets/Input/SSpinBox.h"
 #include "Widgets/Layout/SScrollBox.h"
 #include "Widgets/SCanvas.h"
 #include "Widgets/SNullWidget.h"
@@ -27,122 +25,109 @@
 
 #define LOCTEXT_NAMESPACE "SAdvancedRenamerPanel"
 
-FName FAdvancedRenamerPreviewListItem::OriginalNameColumnName = "OriginalName";
-FName FAdvancedRenamerPreviewListItem::NewNameColumnName = "NewName";
-
 namespace UE::AdvancedRenamer::Private
 {
 	const FVector2D WindowSize = {600.f, 500.f};
 	const FSlateFontInfo TitleFont = FCoreStyle::GetDefaultFontStyle("Regular", 12);
 	const FSlateFontInfo RegularFont = FCoreStyle::GetDefaultFontStyle("Regular", 10);
-	const float LeftBlockWidth = 304.f;
+	constexpr float LeftBlockWidth = 304.f;
 	const FVector2D TitleSize = {150.f, 20.f};
-	const float TitleOffsetY = 25.f;
+	constexpr float TitleOffsetY = 25.f;
 	const FVector2D CheckboxOffset = {0.f, 2.f};
 	const FVector2D CheckboxSize = {18.f, 18.f};
 	const FVector2D RadioOffset = {1.f, 2.f};
-	const float LabelStart = 25.f;
-	const float LabelOffsetY = 3.f;
+	constexpr float LabelStart = 25.f;
+	constexpr float LabelOffsetY = 3.f;
 	const FVector2D LabelSize = {125.f, 18.f};
-	const float EntryStart = 145.f;
+	constexpr float EntryStart = 145.f;
 	const FVector2D EntrySize = {155.f, 21.f};
-	const float LineHeight = 26.f;
+	constexpr float LineHeight = 26.f;
 	const FVector2D SeparatorSize = {32.f, 21.f};
 	const FVector2D SpinSize1 = {32.f, 21.f};
 	const FVector2D SpinSize2 = {39.f, 21.f};
 	const FVector2D SpinSize3 = {46.f, 21.f};
-	const float RightBlockOffsetX = 314.f;
+	constexpr float RightBlockOffsetX = 314.f;
 	const FVector2D RightBlockSize = {281.f, 490.f};
 	const FVector2D ListViewSize = {277.f, 431.f};
-	const float ListLineHeight = 15.f;
-	const float ApplyButtonHeight = 25.f;
+	constexpr float ListLineHeight = 15.f;
+	constexpr float ApplyButtonHeight = 25.f;
+	static FName OriginalNameColumnName = "OriginalName";
+	static FName NewNameColumnName = "NewName";
 }
 
-void SAdvancedRenamerPreviewListRow::Construct(const FArguments& InArgs, TSharedPtr<SAdvancedRenamerPanel> InRenamePanel, 
-	const TSharedRef<STableViewBase>& InOwnerTableView, FObjectRenamePreviewListItemPtr InRowItem)
+class SAdvancedRenamerPreviewListRow : public SMultiColumnTableRow<TSharedPtr<FAdvancedRenamerPreview>>
 {
-	RenamePanel = InRenamePanel;
-	RowItem = InRowItem;
-
-	SMultiColumnTableRow<FObjectRenamePreviewListItemPtr>::Construct(FSuperRowType::FArguments(), InOwnerTableView);
-	SetBorderImage(TAttribute<const FSlateBrush*>(this, &SAdvancedRenamerPreviewListRow::GetBorder));
-}
+public:
+	SLATE_BEGIN_ARGS(SAdvancedRenamerPreviewListRow) {}
+	SLATE_END_ARGS()
 
-TSharedRef<SWidget> SAdvancedRenamerPreviewListRow::GenerateWidgetForColumn(const FName& ColumnName)
-{
-	if (ColumnName != FAdvancedRenamerPreviewListItem::OriginalNameColumnName
-		&& ColumnName != FAdvancedRenamerPreviewListItem::NewNameColumnName)
+	void Construct(const FArguments& InArgs, TSharedPtr<SAdvancedRenamerPanel> InRenamePanel,
+		const TSharedRef<STableViewBase>& InOwnerTableView, TSharedPtr<FAdvancedRenamerPreview> InRowItem)
 	{
-		return SNullWidget::NullWidget;
-	}
-
-	TSharedPtr<SAdvancedRenamerPanel> RenamePanelSP = RenamePanel.Pin();
+		PanelWeak = InRenamePanel;
+		ItemWeak = InRowItem;
 
-	if (!RenamePanelSP.IsValid())
-	{
-		return SNullWidget::NullWidget;
+		SMultiColumnTableRow<TSharedPtr<FAdvancedRenamerPreview>>::Construct(FSuperRowType::FArguments(), InOwnerTableView);
+		SetBorderImage(TAttribute<const FSlateBrush*>(this, &SAdvancedRenamerPreviewListRow::GetBorder));
 	}
 
-	FObjectRenamePreviewListItemPtr RowItemSP = RowItem.Pin();
-
-	if (!RowItemSP.IsValid())
+	virtual TSharedRef<SWidget> GenerateWidgetForColumn(const FName& InColumnName) override
 	{
-		return SNullWidget::NullWidget;
-	}
+		using namespace UE::AdvancedRenamer::Private;
 
-	using namespace UE::AdvancedRenamer::Private;
+		if (InColumnName != OriginalNameColumnName
+			&& InColumnName != NewNameColumnName)
+		{
+			return SNullWidget::NullWidget;
+		}
 
-	TSharedRef<STextBlock> TextBlock = SNew(STextBlock)
-		.ColorAndOpacity(FStyleColors::AccentWhite.GetSpecifiedColor())
-		.Font(RegularFont)
-		.Margin(FMargin(5.f, 0.f, 5.f, 0.f));
-		
-	if (ColumnName == FAdvancedRenamerPreviewListItem::OriginalNameColumnName)
-	{
-		TextBlock->SetText(FText::FromString(RowItemSP->OriginalName));
-		RenamePanelSP->MinDesiredOriginalNameWidth = FMath::Max(RenamePanelSP->MinDesiredOriginalNameWidth, TextBlock->ComputeDesiredSize(1.f).X);
-	}
-	else if (ColumnName == FAdvancedRenamerPreviewListItem::NewNameColumnName)
-	{
-		TextBlock->SetText(FText::FromString(RowItemSP->NewName));
-		RenamePanelSP->MinDesiredNewNameWidth = FMath::Max(RenamePanelSP->MinDesiredOriginalNameWidth, TextBlock->ComputeDesiredSize(1.f).X);
-	}
+		TSharedPtr<SAdvancedRenamerPanel> Panel = PanelWeak.Pin();
 
-	RenamePanelSP->UpdateRequiredListWidth();
+		if (!Panel.IsValid())
+		{
+			return SNullWidget::NullWidget;
+		}
 
-	return TextBlock;
-}
+		TSharedPtr<FAdvancedRenamerPreview> Item = ItemWeak.Pin();
 
-void SAdvancedRenamerPanel::Construct(const FArguments& InArgs)
-{
-	if (InArgs._SharedProvider.IsValid())
-	{
-		SharedProvider = InArgs._SharedProvider;
-	}
-	else
-	{
-		checkNoEntry();
-	}
+		if (!Item.IsValid())
+		{
+			return SNullWidget::NullWidget;
+		}
 
-	int32 Count = Num();
-	check(Count > 0);
+		TSharedRef<STextBlock> TextBlock = SNew(STextBlock)
+			.ColorAndOpacity(FStyleColors::AccentWhite.GetSpecifiedColor())
+			.Font(RegularFont)
+			.Margin(FMargin(5.f, 0.f, 5.f, 0.f));
 
-	for (int32 Index = 0; Index < Count; ++Index)
-	{
-		if (!CanRename(Index))
+		if (InColumnName == OriginalNameColumnName)
 		{
-			RemoveIndex(Index);
-			--Index;
-			--Count;
-			continue;
+			TextBlock->SetText(FText::FromString(Item->OriginalName));
+			Panel->MinDesiredOriginalNameWidth = FMath::Max(Panel->MinDesiredOriginalNameWidth, TextBlock->ComputeDesiredSize(1.f).X);
+		}
+		else if (InColumnName == NewNameColumnName)
+		{
+			TextBlock->SetText(FText::FromString(Item->NewName));
+			Panel->MinDesiredNewNameWidth = FMath::Max(Panel->MinDesiredOriginalNameWidth, TextBlock->ComputeDesiredSize(1.f).X);
 		}
 
-		int32 Hash = GetHash(Index);
-		FString OriginalName = GetOriginalName(Index);
-
-		ListData.Add(MakeShared<FAdvancedRenamerPreviewListItem>(Hash, OriginalName));
+		return TextBlock;
 	}
 
+protected:
+	TWeakPtr<SAdvancedRenamerPanel> PanelWeak;
+	TWeakPtr<FAdvancedRenamerPreview> ItemWeak;
+};
+
+void SAdvancedRenamerPanel::Construct(const FArguments& InArgs, const TSharedRef<IAdvancedRenamer>& InRenamer)
+{
+	Renamer = InRenamer;
+
+	bRemovePrefixSeparator = false;
+	bRemovePrefixNumChars = false;
+	bRemoveSuffixSeparator = false;
+	bRemoveSuffixNumChars = false;
+
 	CommandList = MakeShared<FUICommandList>();
 	CommandList->MapAction(
 		FGenericCommands::Get().Delete,
@@ -152,7 +137,6 @@ void SAdvancedRenamerPanel::Construct(const FArguments& InArgs)
 
 	using namespace UE::AdvancedRenamer::Private;
 
-	// @formatter:off
 	TSharedRef<SCanvas> Canvas = SNew(SCanvas)
 		+ SCanvas::Slot()
 		.Position(FVector2D(0.f, 0.f))
@@ -161,17 +145,14 @@ void SAdvancedRenamerPanel::Construct(const FArguments& InArgs)
 			SNew(SColorBlock)
 			.Color(FStyleColors::Background.GetSpecifiedColor())
 		];
-	// @formatter:on
 
 	CreateLeftPane(Canvas);
 	CreateRightPane(Canvas);
 
-	// @formatter:off
 	ChildSlot
 	[
 		Canvas
 	];
-	// @formatter:on
 
 	if (FEngineAnalytics::IsAvailable())
 	{
@@ -179,46 +160,43 @@ void SAdvancedRenamerPanel::Construct(const FArguments& InArgs)
 	}
 }
 
-void SAdvancedRenamerPanel::CreateLeftPane(TSharedRef<SCanvas> Canvas)
+void SAdvancedRenamerPanel::CreateLeftPane(const TSharedRef<SCanvas>& InCanvas)
 {	
 	using namespace UE::AdvancedRenamer::Private;
 
-	// @formatter:off
-	Canvas->AddSlot()
+	InCanvas->AddSlot()
 		.Position(FVector2D(5.f, 5.f))
 		.Size(FVector2D(LeftBlockWidth, 50.f))
 		[
 			CreateBaseName()
 		];
 
-	Canvas->AddSlot()
+	InCanvas->AddSlot()
 		.Position(FVector2D(5.f, 60.f))
 		.Size(FVector2D(LeftBlockWidth, 102.f))
 		[
 			CreatePrefix()
 		];
 
-	Canvas->AddSlot()
+	InCanvas->AddSlot()
 		.Position(FVector2D(5.f, 167.f))
 		.Size(FVector2D(LeftBlockWidth, 154.f))
 		[
 			CreateSuffix()
 		];
 
-	Canvas->AddSlot()
+	InCanvas->AddSlot()
 		.Position(FVector2D(5.f, 327.f))
 		.Size(FVector2D(LeftBlockWidth, 168.f))
 		[
 			CreateSearchAndReplace()
 		];
-	// @formatter:on
 }
 
 TSharedRef<SWidget> SAdvancedRenamerPanel::CreateBaseName()
 {
 	using namespace UE::AdvancedRenamer::Private;
 
-	// @formatter:off
 	return SNew(SBorder)
 		.BorderImage(FAppStyle::Get().GetBrush("DetailsView.CategoryTop"))
 		.Content()
@@ -234,21 +212,13 @@ TSharedRef<SWidget> SAdvancedRenamerPanel::CreateBaseName()
 				.Text(LOCTEXT("BaseNameTitle", "Name"))
 			]
 			+ SCanvas::Slot()
-			.Position(FVector2D(0.f, TitleOffsetY) + CheckboxOffset)
-			.Size(CheckboxSize)
-			[
-				SAssignNew(BaseNameCheckBox, SCheckBox)
-				.IsChecked(this, &SAdvancedRenamerPanel::IsBaseNameChecked)
-				.OnCheckStateChanged(this, &SAdvancedRenamerPanel::OnBaseNameCheckBoxChanged)
-			]
-			+ SCanvas::Slot()
-			.Position(FVector2D(LabelStart, TitleOffsetY + LabelOffsetY))
+			.Position(FVector2D(LabelStart * 2.f, TitleOffsetY + LabelOffsetY))
 			.Size(LabelSize)
 			[
 				SNew(STextBlock)
 				.ColorAndOpacity(FStyleColors::AccentWhite.GetSpecifiedColor())
 				.Font(RegularFont)
-				.Text(LOCTEXT("BaseName", "Base Name"))
+				.Text(LOCTEXT("SetBaseName", "Set Base Name"))
 			]
 			+ SCanvas::Slot()
 			.Position(FVector2D(EntryStart, TitleOffsetY))
@@ -259,18 +229,15 @@ TSharedRef<SWidget> SAdvancedRenamerPanel::CreateBaseName()
 				.ForegroundColor(FStyleColors::AccentWhite.GetSpecifiedColor())
 				.Font(RegularFont)
 				.HintText(LOCTEXT("BaseNameHint", "Base name"))
-				.IsEnabled(false)
 				.OnTextChanged(this, &SAdvancedRenamerPanel::OnBaseNameChanged)
 			]
 		];
-	// @formatter:on
 }
 
 TSharedRef<SWidget> SAdvancedRenamerPanel::CreatePrefix()
 {
 	using namespace UE::AdvancedRenamer::Private;
 
-	// @formatter:off
 	return SNew(SBorder)
 		.BorderImage(FAppStyle::Get().GetBrush("DetailsView.CategoryTop"))
 		.Content()
@@ -287,24 +254,15 @@ TSharedRef<SWidget> SAdvancedRenamerPanel::CreatePrefix()
 				.Font(TitleFont)
 				.Text(LOCTEXT("PrefixTitle", "Prefix"))
 			]
-
-			// Prefix name
+			//
 			+ SCanvas::Slot()
-			.Position(FVector2D(0.f, TitleOffsetY) + CheckboxOffset)
-			.Size(CheckboxSize)
-			[
-				SAssignNew(PrefixCheckBox, SCheckBox)
-				.IsChecked(this, &SAdvancedRenamerPanel::IsPrefixChecked)
-				.OnCheckStateChanged(this, &SAdvancedRenamerPanel::OnPrefixCheckBoxChanged)
-			]
-			+ SCanvas::Slot()
-			.Position(FVector2D(LabelStart, TitleOffsetY + LabelOffsetY))
+			.Position(FVector2D(LabelStart * 2.f, TitleOffsetY + LabelOffsetY))
 			.Size(LabelSize)
 			[
 				SNew(STextBlock)
 				.ColorAndOpacity(FStyleColors::AccentWhite.GetSpecifiedColor())
 				.Font(RegularFont)
-				.Text(LOCTEXT("Prefix", "Prefix"))
+				.Text(LOCTEXT("AddPrefix", "Add Prefix"))
 			]
 			+ SCanvas::Slot()
 			.Position(FVector2D(EntryStart, TitleOffsetY))
@@ -315,7 +273,6 @@ TSharedRef<SWidget> SAdvancedRenamerPanel::CreatePrefix()
 				.ForegroundColor(FStyleColors::AccentWhite.GetSpecifiedColor())
 				.Font(RegularFont)
 				.HintText(LOCTEXT("PrefixHint", "Prefix"))
-				.IsEnabled(false)
 				.OnTextChanged(this, &SAdvancedRenamerPanel::OnPrefixChanged)
 			]
 
@@ -328,7 +285,6 @@ TSharedRef<SWidget> SAdvancedRenamerPanel::CreatePrefix()
 				.Style(&FAdvancedRenamerStyle::Get().GetWidgetStyle<FCheckBoxStyle>("AdvancedRenamer.Style.BlackRadioButton"))
 				.IsChecked(this, &SAdvancedRenamerPanel::IsPrefixRemoveChecked)
 				.OnCheckStateChanged(this, &SAdvancedRenamerPanel::OnPrefixRemoveCheckBoxChanged)
-				.IsEnabled(false)
 			]
 			+ SCanvas::Slot()
 			.Position(FVector2D(LabelStart * 2.f, TitleOffsetY + LabelOffsetY + LineHeight))
@@ -347,7 +303,7 @@ TSharedRef<SWidget> SAdvancedRenamerPanel::CreatePrefix()
 				.BackgroundColor(FStyleColors::Background.GetSpecifiedColor())
 				.ForegroundColor(FStyleColors::AccentWhite.GetSpecifiedColor())
 				.Font(RegularFont)
-				.IsEnabled(false)
+				.IsEnabled(this, &SAdvancedRenamerPanel::IsPrefixRemoveSeparatorEnabled)
 				.Text(LOCTEXT("Underscore", "_"))
 				.OnVerifyTextChanged(this, &SAdvancedRenamerPanel::OnPrefixSeparatorVerifyTextChanged)
 				.OnTextChanged(this, &SAdvancedRenamerPanel::OnPrefixSeparatorChanged)
@@ -371,7 +327,6 @@ TSharedRef<SWidget> SAdvancedRenamerPanel::CreatePrefix()
 				.Style(&FAdvancedRenamerStyle::Get().GetWidgetStyle<FCheckBoxStyle>("AdvancedRenamer.Style.BlackRadioButton"))
 				.IsChecked(this, &SAdvancedRenamerPanel::IsPrefixRemoveCharactersChecked)
 				.OnCheckStateChanged(this, &SAdvancedRenamerPanel::OnPrefixRemoveCharactersCheckBoxChanged)
-				.IsEnabled(false)
 			]
 			+ SCanvas::Slot()
 			.Position(FVector2D(LabelStart * 2.f, TitleOffsetY + LabelOffsetY + LineHeight * 2.f))
@@ -393,7 +348,7 @@ TSharedRef<SWidget> SAdvancedRenamerPanel::CreatePrefix()
 				.MinValue(1)
 				.MaxValue(9)
 				.Value(1)
-				.IsEnabled(false)
+				.IsEnabled(this, &SAdvancedRenamerPanel::IsPrefixRemoveNumCharsEnabled)
 				.OnValueChanged(this, &SAdvancedRenamerPanel::OnPrefixRemoveCharactersChanged)
 			]
 			+ SCanvas::Slot()
@@ -406,14 +361,12 @@ TSharedRef<SWidget> SAdvancedRenamerPanel::CreatePrefix()
 				.Text(LOCTEXT("Characters", "character(s)"))
 			]
 		];
-	// @formatter:on
 }
 
 TSharedRef<SWidget> SAdvancedRenamerPanel::CreateSuffix()
 {
 	using namespace UE::AdvancedRenamer::Private;
 
-	// @formatter:off
 	return SNew(SBorder)
 		.BorderImage(FAppStyle::Get().GetBrush("DetailsView.CategoryTop"))
 		.Content()
@@ -433,21 +386,13 @@ TSharedRef<SWidget> SAdvancedRenamerPanel::CreateSuffix()
 
 			// Suffix name
 			+ SCanvas::Slot()
-			.Position(FVector2D(0.f, TitleOffsetY) + CheckboxOffset)
-			.Size(CheckboxSize)
-			[
-				SAssignNew(SuffixCheckBox, SCheckBox)
-				.IsChecked(this, &SAdvancedRenamerPanel::IsSuffixChecked)
-				.OnCheckStateChanged(this, &SAdvancedRenamerPanel::OnSuffixCheckBoxChanged)
-			]
-			+ SCanvas::Slot()
-			.Position(FVector2D(LabelStart, TitleOffsetY + LabelOffsetY))
+			.Position(FVector2D(LabelStart * 2.f, TitleOffsetY + LabelOffsetY))
 			.Size(LabelSize)
 			[
 				SNew(STextBlock)
 				.ColorAndOpacity(FStyleColors::AccentWhite.GetSpecifiedColor())
 				.Font(RegularFont)
-				.Text(LOCTEXT("Suffix", "Suffix"))
+				.Text(LOCTEXT("AddSuffix", "Add Suffix"))
 			]
 			+ SCanvas::Slot()
 			.Position(FVector2D(EntryStart, TitleOffsetY))
@@ -458,7 +403,6 @@ TSharedRef<SWidget> SAdvancedRenamerPanel::CreateSuffix()
 				.ForegroundColor(FStyleColors::AccentWhite.GetSpecifiedColor())
 				.Font(RegularFont)
 				.HintText(LOCTEXT("SuffixHint", "Suffix"))
-				.IsEnabled(false)
 				.OnTextChanged(this, &SAdvancedRenamerPanel::OnSuffixChanged)
 			]
 
@@ -471,7 +415,6 @@ TSharedRef<SWidget> SAdvancedRenamerPanel::CreateSuffix()
 				.Style(&FAdvancedRenamerStyle::Get().GetWidgetStyle<FCheckBoxStyle>("AdvancedRenamer.Style.BlackRadioButton"))
 				.IsChecked(this, &SAdvancedRenamerPanel::IsSuffixRemoveChecked)
 				.OnCheckStateChanged(this, &SAdvancedRenamerPanel::OnSuffixRemoveCheckBoxChanged)
-				.IsEnabled(false)
 			]
 			+ SCanvas::Slot()
 			.Position(FVector2D(LabelStart * 2.f, TitleOffsetY + LabelOffsetY + LineHeight))
@@ -490,7 +433,7 @@ TSharedRef<SWidget> SAdvancedRenamerPanel::CreateSuffix()
 				.BackgroundColor(FStyleColors::Background.GetSpecifiedColor())
 				.ForegroundColor(FStyleColors::AccentWhite.GetSpecifiedColor())
 				.Font(RegularFont)
-				.IsEnabled(false)
+				.IsEnabled(this, &SAdvancedRenamerPanel::IsSuffixRemoveSeparatorEnabled)
 				.Text(LOCTEXT("Underscore", "_"))
 				.OnVerifyTextChanged(this, &SAdvancedRenamerPanel::OnSuffixSeparatorVerifyTextChanged)
 				.OnTextChanged(this, &SAdvancedRenamerPanel::OnSuffixSeparatorChanged)
@@ -514,7 +457,6 @@ TSharedRef<SWidget> SAdvancedRenamerPanel::CreateSuffix()
 				.Style(&FAdvancedRenamerStyle::Get().GetWidgetStyle<FCheckBoxStyle>("AdvancedRenamer.Style.BlackRadioButton"))
 				.IsChecked(this, &SAdvancedRenamerPanel::IsSuffixRemoveCharactersChecked)
 				.OnCheckStateChanged(this, &SAdvancedRenamerPanel::OnSuffixRemoveCharactersCheckBoxChanged)
-				.IsEnabled(false)
 			]
 			+ SCanvas::Slot()
 			.Position(FVector2D(LabelStart * 2.f, TitleOffsetY + LabelOffsetY + LineHeight * 2.f))
@@ -529,14 +471,13 @@ TSharedRef<SWidget> SAdvancedRenamerPanel::CreateSuffix()
 			.Position(FVector2D(EntryStart, TitleOffsetY + LineHeight * 2.f))
 			.Size(SpinSize1)
 			[
-
 				SAssignNew(SuffixRemoveCharactersSpinBox, SSpinBox<uint8>)
 				.Style(&FAppStyle::Get(), "Menu.SpinBox")
 				.Font(RegularFont)
 				.MinValue(1)
 				.MaxValue(9)
 				.Value(1)
-				.IsEnabled(false)
+				.IsEnabled(this, &SAdvancedRenamerPanel::IsSuffixRemoveNumCharsEnabled)
 				.OnValueChanged(this, &SAdvancedRenamerPanel::OnSuffixRemoveCharactersChanged)
 			]
 			+ SCanvas::Slot()
@@ -557,7 +498,7 @@ TSharedRef<SWidget> SAdvancedRenamerPanel::CreateSuffix()
 				SAssignNew(SuffixRemoveNumberCheckBox, SCheckBox)
 				.IsChecked(this, &SAdvancedRenamerPanel::IsSuffixRemoveNumberChecked)
 				.OnCheckStateChanged(this, &SAdvancedRenamerPanel::OnSuffixRemoveNumberCheckBoxChanged)
-				.IsEnabled(false)
+				.IsEnabled(this, &SAdvancedRenamerPanel::IsSuffixRemoveNumberCheckBoxEnabled)
 			]
 			+ SCanvas::Slot()
 			.Position(FVector2D(LabelStart * 2.f, TitleOffsetY + LabelOffsetY + LineHeight * 3.f))
@@ -577,7 +518,6 @@ TSharedRef<SWidget> SAdvancedRenamerPanel::CreateSuffix()
 				SAssignNew(SuffixNumberCheckBox, SCheckBox)
 				.IsChecked(this, &SAdvancedRenamerPanel::IsSuffixNumberChecked)
 				.OnCheckStateChanged(this, &SAdvancedRenamerPanel::OnSuffixNumberCheckBoxChanged)
-				.IsEnabled(false)
 			]
 			+ SCanvas::Slot()
 			.Position(FVector2D(LabelStart * 2.f, TitleOffsetY + LabelOffsetY + LineHeight * 4.f))
@@ -586,7 +526,7 @@ TSharedRef<SWidget> SAdvancedRenamerPanel::CreateSuffix()
 				SNew(STextBlock)
 				.ColorAndOpacity(FStyleColors::AccentWhite.GetSpecifiedColor())
 				.Font(RegularFont)
-				.Text(LOCTEXT("SuffixNumber", "Number"))
+				.Text(LOCTEXT("SuffixNumber", "Add Number"))
 			]
 			+ SCanvas::Slot()
 			.Position(FVector2D(EntryStart , TitleOffsetY + LabelOffsetY + LineHeight * 4.f))
@@ -635,14 +575,12 @@ TSharedRef<SWidget> SAdvancedRenamerPanel::CreateSuffix()
 				.OnValueChanged(this, &SAdvancedRenamerPanel::OnSuffixNumberStepChanged)
 			]
 		];
-	// @formatter:on
 }
 
 TSharedRef<SWidget> SAdvancedRenamerPanel::CreateSearchAndReplace()
 {
 	using namespace UE::AdvancedRenamer::Private;
 
-	// @formatter:off
 	return SNew(SBorder)
 		.BorderImage(FAppStyle::Get().GetBrush("DetailsView.CategoryTop"))
 		.Content()
@@ -731,7 +669,6 @@ TSharedRef<SWidget> SAdvancedRenamerPanel::CreateSearchAndReplace()
 				.AutoWrapText(true)
 				.HintText(LOCTEXT("RegeSearchHint", "Search"))
 				.WrappingPolicy(ETextWrappingPolicy::AllowPerCharacterWrapping)
-				.IsEnabled(false)
 				.OnTextChanged(this, &SAdvancedRenamerPanel::OnSearchReplaceSearchTextChanged)
 			]
 
@@ -747,42 +684,38 @@ TSharedRef<SWidget> SAdvancedRenamerPanel::CreateSearchAndReplace()
 				.AutoWrapText(true)
 				.HintText(LOCTEXT("RegeReplaceHint", "Replace"))
 				.WrappingPolicy(ETextWrappingPolicy::AllowPerCharacterWrapping)
-				.IsEnabled(false)
 				.OnTextChanged(this, &SAdvancedRenamerPanel::OnSearchReplaceReplaceTextChanged)
 			]
 		];
-	// @formatter:on
 }
 
-void SAdvancedRenamerPanel::CreateRightPane(TSharedRef<SCanvas> Canvas)
+void SAdvancedRenamerPanel::CreateRightPane(const TSharedRef<SCanvas>& InCanvas)
 {
 	using namespace UE::AdvancedRenamer::Private;
 
-	// @formatter:off
-	Canvas->AddSlot()
+	InCanvas->AddSlot()
 		.Position(FVector2D(RightBlockOffsetX, 5.f))
-		.Size(RightBlockSize)
+		.Size(this, &SAdvancedRenamerPanel::GetRightPaneSize)
 		[
 			CreateRenamePreview()
 		];
-	// @formatter:on
 }
 
 TSharedRef<SWidget> SAdvancedRenamerPanel::CreateRenamePreview()
 {
 	using namespace UE::AdvancedRenamer::Private;
 
-	// @formatter:off
 	return SNew(SBorder)
 		.BorderImage(FAppStyle::Get().GetBrush("DetailsView.CategoryTop"))
+		.HAlign(EHorizontalAlignment::HAlign_Fill)
+		.VAlign(EVerticalAlignment::VAlign_Fill)
 		.Content()
 		[
-			SNew(SCanvas)
-				
-			// Title
-			+ SCanvas::Slot()
-			.Position(FVector2D(0.f, LabelOffsetY))
-			.Size(TitleSize)
+			SNew(SVerticalBox)
+
+			+ SVerticalBox::Slot()
+			.Padding(0.f, LabelOffsetY, 0.f, 0.f)
+			.AutoHeight()
 			[
 				SNew(STextBlock)
 				.ColorAndOpacity(FStyleColors::AccentBlue.GetSpecifiedColor())
@@ -790,52 +723,41 @@ TSharedRef<SWidget> SAdvancedRenamerPanel::CreateRenamePreview()
 				.Text(LOCTEXT("ObjectRenamePreviewTitle", "Rename Preview"))
 			]
 
-			// List View
-			+ SCanvas::Slot()
-			.Position(FVector2D(0.f, TitleOffsetY))
-			.Size(ListViewSize)
+			+ SVerticalBox::Slot()
+			.FillHeight(1.f)
+			.HAlign(EHorizontalAlignment::HAlign_Fill)
 			[
-				SNew(SScrollBox)
-				.Orientation(EOrientation::Orient_Horizontal)
-				+ SScrollBox::Slot()
+				SAssignNew(RenamePreviewListBox, SBox)
+				.Content()
 				[
-					SAssignNew(RenamePreviewListBox, SBox)
-					.WidthOverride(ListViewSize.Y)
-					.HeightOverride(ListViewSize.X)
-					.HAlign(EHorizontalAlignment::HAlign_Fill)
-					.VAlign(EVerticalAlignment::VAlign_Fill)
-					.Content()
-					[
-						SAssignNew(RenamePreviewList, SListView<FObjectRenamePreviewListItemPtr>)
-						.ItemHeight(ListLineHeight)
-						.ListItemsSource(&ListData)
-						.OnGenerateRow(this, &SAdvancedRenamerPanel::OnGenerateRowForList)
-						.HeaderRow(
-							SAssignNew(RenamePreviewListHeaderRow, SHeaderRow)
-							+ SHeaderRow::Column(FAdvancedRenamerPreviewListItem::OriginalNameColumnName)
-							.DefaultLabel(LOCTEXT("From", "From"))
-							.FillWidth(ListViewSize.X / 2.f)
-							+ SHeaderRow::Column(FAdvancedRenamerPreviewListItem::NewNameColumnName)
-							.DefaultLabel(LOCTEXT("To", "To"))
-							.FillWidth(ListViewSize.X / 2.f)
-						)
-						.OnKeyDownHandler(this, &SAdvancedRenamerPanel::OnListViewKeyDown)
-						.OnContextMenuOpening(this, &SAdvancedRenamerPanel::GenerateListViewContextMenu)
-					]
+					SAssignNew(RenamePreviewList, SListView<TSharedPtr<FAdvancedRenamerPreview>>)
+					.ItemHeight(ListLineHeight)
+					.ListItemsSource(&Renamer->GetPreviews())
+					.OnGenerateRow(this, &SAdvancedRenamerPanel::OnGenerateRowForList)
+					.HeaderRow(
+						SAssignNew(RenamePreviewListHeaderRow, SHeaderRow)
+						+ SHeaderRow::Column(OriginalNameColumnName)
+						.DefaultLabel(LOCTEXT("From", "From"))
+						.FillWidth(0.5f)
+						+ SHeaderRow::Column(NewNameColumnName)
+						.DefaultLabel(LOCTEXT("To", "To"))
+						.FillWidth(0.5f)
+					)
+					.OnKeyDownHandler(this, &SAdvancedRenamerPanel::OnListViewKeyDown)
+					.OnContextMenuOpening(this, &SAdvancedRenamerPanel::GenerateListViewContextMenu)
 				]
 			]
 
-			// Apply Button
-			+ SCanvas::Slot()
-			.Position(FVector2D(0.f, TitleOffsetY + ListViewSize.Y + 5.f))
-			.Size(FVector2D(ListViewSize.X, ApplyButtonHeight))
+			+ SVerticalBox::Slot()
+			.AutoHeight()
+			.HAlign(EHorizontalAlignment::HAlign_Fill)
 			[
 				SAssignNew(ApplyButton, SButton)
 				.ButtonStyle(FAdvancedRenamerStyle::Get(), "AdvancedRenamer.Style.DarkButton")
 				.ToolTipText(LOCTEXT("ApplyRenameToolTip", "Renames all actors."))
 				.HAlign(EHorizontalAlignment::HAlign_Center)
 				.VAlign(EVerticalAlignment::VAlign_Center)
-				.IsEnabled(false)
+				.IsEnabled(this, &SAdvancedRenamerPanel::IsApplyButtonEnabled)
 				.OnClicked(this, &SAdvancedRenamerPanel::OnApplyButtonClicked)
 				[
 					SNew(STextBlock)
@@ -845,52 +767,6 @@ TSharedRef<SWidget> SAdvancedRenamerPanel::CreateRenamePreview()
 				]
 			]
 		];
-	// @formatter:on
-}
-
-bool SAdvancedRenamerPanel::RenameObjects()
-{
-	if (!bValidNames)
-	{
-		return false;
-	}
-
-	FScopedTransaction Transaction(LOCTEXT("AdvancedObjectRename", "Advanced Object Rename"));
-	int32 ItemsRenamed = 0;
-	int32 ItemCount = Num();
-	
-	for (int32 Index = 0; Index < ItemCount; ++Index)
-	{
-		if (!ListData[Index].IsValid())
-		{
-			continue;
-		}
-
-		if (!IsValidIndex(Index))
-		{
-			continue;
-		}
-
-		if (ListData[Index]->NewName.Len() == 0)
-		{
-			continue;
-		}
-
-		if (ExecuteRename(Index, ListData[Index]->NewName))
-		{
-			++ItemsRenamed;
-		}
-	}
-
-	//Cancel Transaction if no Items could successfully Rename
-	if (ItemsRenamed == 0)
-	{
-		Transaction.Cancel();
-		//TODO: Should we keep Window Open as no renames occurred. Might be something the User did not expect
-		//return false;
-	}
-	
-	return true;
 }
 
 bool SAdvancedRenamerPanel::CloseWindow()
@@ -906,512 +782,22 @@ bool SAdvancedRenamerPanel::CloseWindow()
 	return false;
 }
 
-FString SAdvancedRenamerPanel::CreateNewName(int32 Index) const
-{
-	if (!IsValidIndex(Index))
-	{
-		return "";
-	}
-
-	FString DisplayName = GetOriginalName(Index);
-
-	if (DisplayName.Len() == 0)
-	{
-		return DisplayName;
-	}
-
-	return ApplyRename(DisplayName, Index);
-}
-
-FString SAdvancedRenamerPanel::ApplyRename(const FString& OriginalName, int32 Index) const
-{
-	FString NewName = OriginalName;
-
-	if (bBaseName)
-	{
-		NewName = ApplyBaseName(NewName);
-	}
-
-	if (bPrefix)
-	{
-		NewName = ApplyPrefix(NewName);
-	}
-
-	if (bSuffix)
-	{
-		NewName = ApplySuffix(NewName, Index);
-	}
-
-	if (bSearchReplacePlainText)
-	{
-		NewName = ApplySearchPlainText(NewName);
-	}
-	else if (bSearchReplaceRegex)
-	{
-		NewName = ApplySearchReplaceRegex(NewName);
-	}
-
-	return NewName;
-}
-
-FString SAdvancedRenamerPanel::ApplyBaseName(const FString& OriginalName) const
-{
-	if (!BaseNameTextBox.IsValid())
-	{
-		return OriginalName;
-	}
-
-	FString BaseName = BaseNameTextBox->GetText().ToString();
-	
-	if (BaseName.Len() == 0)
-	{
-		return OriginalName;
-	}
-
-	return BaseName;
-}
-
-FString SAdvancedRenamerPanel::ApplyPrefix(const FString& OriginalName) const
-{
-	FString Output = OriginalName;
-
-	if (bPrefixRemove)
-	{
-		if (PrefixSeparatorTextBox.IsValid())
-		{
-			FString Separator = PrefixSeparatorTextBox->GetText().ToString();
-
-			if (Separator.Len() != 0)
-			{
-				int32 PrefixStart = Output.Find(Separator, ESearchCase::IgnoreCase);
-
-				if (PrefixStart >= 0)
-				{
-					Output = Output.Mid(PrefixStart + Separator.Len());
-				}
-			}
-		}
-	}
-	else if (bPrefixRemoveCharacters)
-	{
-		if (PrefixRemoveCharactersSpinBox.IsValid())
-		{
-			Output = Output.RightChop(PrefixRemoveCharactersSpinBox->GetValue());
-		}
-	}
-
-	if (PrefixTextBox.IsValid())
-	{
-		FString Prefix = PrefixTextBox->GetText().ToString();
-
-		if (Prefix.Len() > 0)
-		{
-			Output = Prefix + Output;
-		}
-	}
-
-	return Output;
-}
-
-FString SAdvancedRenamerPanel::ApplySuffix(const FString& OriginalName, int32 Index) const
-{
-	static const TCHAR FirstDigit = '0';
-	static const TCHAR LastDigit = '9';
-
-	FString Output = OriginalName;
-
-	if (bSuffixRemove)
-	{
-		if (SuffixSeparatorTextBox.IsValid())
-		{
-			FString Separator = SuffixSeparatorTextBox->GetText().ToString();
-
-			if (Separator.Len() != 0)
-			{
-				int32 SuffixStart = Output.Find(Separator, ESearchCase::IgnoreCase, ESearchDir::FromEnd);
-
-				if (SuffixStart >= 0)
-				{
-					Output = Output.Mid(0, SuffixStart);
-				}
-			}
-		}
-	}
-	else if (bSuffixRemoveCharacters)
-	{
-		if (SuffixRemoveCharactersSpinBox.IsValid())
-		{
-			Output = Output.LeftChop(SuffixRemoveCharactersSpinBox->GetValue());
-		}
-	}
-
-	if (bSuffixRemoveNumber)
-	{
-		int32 LastDigitIndex = Output.Len() - 1;
-
-		while(LastDigitIndex >= 0)
-		{
-			if (Output[LastDigitIndex] < FirstDigit || Output[LastDigitIndex] > LastDigit)
-			{
-				break;
-			}
-
-			--LastDigitIndex;
-		}
-
-		Output = Output.Mid(0, LastDigitIndex + 1);
-	}
-
-	if (SuffixTextBox.IsValid())
-	{
-		FString Suffix = SuffixTextBox->GetText().ToString();
-
-		if (Suffix.Len() > 0)
-		{
-			Output = Output + Suffix;
-		}
-	}
-
-	if (bSuffixNumber)
-	{
-		if (SuffixNumberStartSpinBox.IsValid() && SuffixNumberStepSpinBox.IsValid())
-		{
-			int32 Start = SuffixNumberStartSpinBox->GetValue();
-			int32 Step = SuffixNumberStepSpinBox->GetValue();
-			int32 Current = Start + (Index * Step);
-
-			Output += FString::FromInt(Current);
-		}
-	}
-
-	return Output;
-}
-
-FString SAdvancedRenamerPanel::ApplySearchPlainText(const FString& OriginalName) const
-{
-	if (!SearchReplaceSearchTextBox.IsValid() || !SearchReplaceReplaceTextBox.IsValid())
-	{
-		return OriginalName;
-	}
-
-	FString Search = SearchReplaceSearchTextBox->GetPlainText().ToString();
-
-	if (Search.Len() == 0)
-	{
-		return OriginalName;
-	}
-
-	FString Replace = SearchReplaceReplaceTextBox->GetPlainText().ToString();
-
-	if (Replace.Len() == 0)
-	{
-		UE_LOG(LogARP, Warning, TEXT("Regex: Empty replacement string."));
-	}
-
-	return OriginalName.Replace(*Search, *Replace, bSearchReplaceIgnoreCase ? ESearchCase::IgnoreCase : ESearchCase::CaseSensitive);
-}
-
-FString SAdvancedRenamerPanel::ApplySearchReplaceRegex(const FString& OriginalName) const
-{
-	if (!SearchReplaceSearchTextBox.IsValid() || !SearchReplaceReplaceTextBox.IsValid())
-	{
-		return OriginalName;
-	}
-
-	FString Pattern = SearchReplaceSearchTextBox->GetPlainText().ToString();
-
-	if (Pattern.Len() == 0)
-	{
-		return OriginalName;
-	}
-	
-	FRegexPattern RegexPattern = FRegexPattern(
-		Pattern, 
-		bSearchReplaceIgnoreCase ? ERegexPatternFlags::CaseInsensitive : ERegexPatternFlags::None
-	);
-
-	FString ReplaceString = SearchReplaceReplaceTextBox->GetPlainText().ToString();
-
-	if (ReplaceString.Len() == 0)
-	{
-		UE_LOG(LogARP, Warning, TEXT("Regex: Empty replacement string."));
-	}
-
-	return RegexReplace(OriginalName, RegexPattern, ReplaceString);
-}
-
-FString SAdvancedRenamerPanel::RegexReplace(const FString& OriginalString, const FRegexPattern Pattern, const FString& ReplaceString) const
-{
-	static const FString EscapeString = "\\";
-	static const TCHAR EscapeChar = EscapeString[0];
-	static const FString GroupString = "$";
-	static const TCHAR GroupChar = GroupString[0];
-	static const TCHAR FirstDigit = '0';
-	static const TCHAR LastDigit = '9';
-	static const TCHAR NullChar = 0;
-
-	FRegexMatcher Matcher(Pattern, OriginalString);
-
-	FString Output = "";
-	int32 StartCharIdx = 0;
-	bool bReadingGroupName = false;
-	int32 GroupIndex = INDEX_NONE;
-	
-	while (Matcher.FindNext())
-	{
-		// Add on part of string after start/previous match
-		if (StartCharIdx != Matcher.GetMatchBeginning())
-		{
-			Output += OriginalString.Mid(StartCharIdx, Matcher.GetMatchBeginning() - StartCharIdx);
-		}
-
-		bool bEscaped = false;
-
-		for (int32 CharIdx = 0; CharIdx <= ReplaceString.Len(); ++CharIdx)
-		{
-			const TCHAR& Char = CharIdx < ReplaceString.Len() ? ReplaceString[CharIdx] : NullChar;
-
-			if (bReadingGroupName)
-			{
-				// Build group index
-				if (Char >= FirstDigit && Char <= LastDigit)
-				{
-					if (GroupIndex == INDEX_NONE)
-					{
-						GroupIndex = 0;
-					}
-					else if (GroupIndex > 0)
-					{
-						GroupIndex *= 10;
-					}
-
-					int32 NextDigit = static_cast<int32>(Char - FirstDigit);					
-					GroupIndex += NextDigit;
-					continue;
-				}
-				// We've read a group index, add it to the output string
-				else if (GroupIndex > 0)
-				{
-					if (Matcher.GetCaptureGroupBeginning(GroupIndex) == INDEX_NONE)
-					{
-						UE_LOG(LogARP, Error, TEXT("Regex: Capture group does not exist %d."), GroupIndex);
-					}
-
-					Output += Matcher.GetCaptureGroup(GroupIndex);
-				}
-				// $0 matches the entire matched string
-				else if (GroupIndex == 0)
-				{
-					Output += OriginalString.Mid(Matcher.GetMatchBeginning(), Matcher.GetMatchEnding() - Matcher.GetMatchBeginning());
-				}
-				// An unescaped $
-				else
-				{
-					UE_LOG(LogARP, Error, TEXT("Regex: Unescaped %s."), *GroupString);
-
-					Output += GroupString;
-				}
-
-				bReadingGroupName = false;
-				// Continue regular parsing of this character.
-			}
-
-			// Check for special chars
-			if (!bEscaped)
-			{
-				if (Char == EscapeChar)
-				{
-					bEscaped = true;
-					continue;
-				}
-
-				if (Char == GroupChar)
-				{
-					bReadingGroupName = true;
-					GroupIndex = INDEX_NONE;
-					continue;
-				}
-			}
-			else
-			{
-				// If the last char is a \ assume that it's not an escape char
-				if (Char == NullChar)
-				{
-					UE_LOG(LogARP, Error, TEXT("Regex: Unescaped %s."), *EscapeString);
-
-					Output += EscapeChar;
-				}
-			}
-
-			if (Char != NullChar)
-			{
-				Output += Char;
-			}
-
-			bEscaped = false;
-		}
-
-		StartCharIdx = Matcher.GetMatchEnding();
-	}
-
-	// Add on the end of the string after the last match
-	if (StartCharIdx < OriginalString.Len())
-	{
-		Output += OriginalString.Mid(StartCharIdx);
-	}
-
-	return Output;
-}
-
-void SAdvancedRenamerPanel::UpdateEnables()
-{
-	if (BaseNameTextBox.IsValid())
-	{
-		BaseNameTextBox->SetEnabled(bBaseName);
-	}
-
-	if (PrefixTextBox.IsValid())
-	{
-		PrefixTextBox->SetEnabled(bPrefix);
-	}
-
-	if (PrefixRemoveCheckBox.IsValid())
-	{
-		PrefixRemoveCheckBox->SetEnabled(bPrefix);
-	}
-
-	if (PrefixSeparatorTextBox.IsValid())
-	{
-		PrefixSeparatorTextBox->SetEnabled(bPrefix && bPrefixRemove);
-	}
-
-	if (PrefixRemoveCharactersCheckBox.IsValid())
-	{
-		PrefixRemoveCharactersCheckBox->SetEnabled(bPrefix);
-	}
-
-	if (PrefixRemoveCharactersSpinBox.IsValid())
-	{
-		PrefixRemoveCharactersSpinBox->SetEnabled(bPrefix && bPrefixRemoveCharacters);
-	}
-
-	if (SuffixTextBox.IsValid())
-	{
-		SuffixTextBox->SetEnabled(bSuffix);
-	}
-
-	if (SuffixRemoveCheckBox.IsValid())
-	{
-		SuffixRemoveCheckBox->SetEnabled(bSuffix);
-	}
-
-	if (SuffixSeparatorTextBox.IsValid())
-	{
-		SuffixSeparatorTextBox->SetEnabled(bSuffix && bSuffixRemove);
-	}
-
-	if (SuffixRemoveCharactersCheckBox.IsValid())
-	{
-		SuffixRemoveCharactersCheckBox->SetEnabled(bSuffix);
-	}
-
-	if (SuffixRemoveCharactersSpinBox.IsValid())
-	{
-		SuffixRemoveCharactersSpinBox->SetEnabled(bSuffix && bSuffixRemoveCharacters);
-	}
-
-	if (SuffixRemoveNumberCheckBox.IsValid())
-	{
-		SuffixRemoveNumberCheckBox->SetEnabled(bSuffix && !bSuffixNumber);
-	}
-
-	if (SuffixNumberCheckBox.IsValid())
-	{
-		SuffixNumberCheckBox->SetEnabled(bSuffix);
-	}
-
-	if (SuffixNumberStartSpinBox.IsValid())
-	{
-		SuffixNumberStartSpinBox->SetEnabled(bSuffix && bSuffixNumber);
-	}
-
-	if (SuffixNumberStepSpinBox.IsValid())
-	{
-		SuffixNumberStepSpinBox->SetEnabled(bSuffix && bSuffixNumber);
-	}
-
-	if (SearchReplaceIgnoreCaseCheckBox.IsValid())
-	{
-		SearchReplaceIgnoreCaseCheckBox->SetEnabled(bSearchReplacePlainText || bSearchReplaceRegex);
-	}
-
-	if (SearchReplaceSearchTextBox.IsValid())
-	{
-		SearchReplaceSearchTextBox->SetEnabled(bSearchReplacePlainText || bSearchReplaceRegex);
-	}
-
-	if (SearchReplaceReplaceTextBox.IsValid())
-	{
-		SearchReplaceReplaceTextBox->SetEnabled(bSearchReplacePlainText || bSearchReplaceRegex);
-	}
-
-	if (ApplyButton.IsValid())
-	{
-		ApplyButton->SetEnabled(bValidNames);
-	}
-
-	RequestListViewRefresh();
-}
-
-void SAdvancedRenamerPanel::RequestListViewRefresh()
-{
-	bListNeedsUpdate = true;
-}
-
 void SAdvancedRenamerPanel::RefreshListView(const double InCurrentTime)
 {
-	bool bRemovedObjects = false;
-	bValidNames = false;
-
-	for (int32 Index = 0; Index < ListData.Num(); ++Index)
-	{
-		if (!ListData[Index].IsValid() || !IsValidIndex(Index))
-		{
-			RemoveIndex(Index);
-			--Index;
-			bRemovedObjects = true;
-			continue;
-		}
-
-		// Force recreation
-		ListData[Index]->NewName = CreateNewName(Index);
-
-		if (ListData[Index]->NewName.Len() == 0)
-		{
-			continue;
-		}
+	const int32 CurrentCount = Renamer->Num();
 
-		if (GetOriginalName(Index) == ListData[Index]->NewName)
-		{
-			continue;
-		}
-
-		bValidNames = true;
-	}
+	Renamer->UpdatePreviews();
 
-	if (bRemovedObjects)
+	if (CurrentCount != Renamer->Num())
 	{
 		RenamePreviewList->RequestListRefresh();
 	}
 
-	UpdateEnables();
-
 	MinDesiredOriginalNameWidth = 0.f;
 	MinDesiredNewNameWidth = 0.f;
 	RenamePreviewList->RebuildList();
 
 	ListLastUpdateTime = InCurrentTime;
-	bListNeedsUpdate = false;
 }
 
 void SAdvancedRenamerPanel::UpdateRequiredListWidth()
@@ -1426,22 +812,19 @@ void SAdvancedRenamerPanel::UpdateRequiredListWidth()
 		return;
 	}
 
-	static const float ListViewWidth = 277.f;
-	float ActualWidth = FMath::Max(ListViewWidth, MinDesiredOriginalNameWidth + MinDesiredNewNameWidth + 20.f);
+	using namespace UE::AdvancedRenamer::Private;
+
+	const float ActualWidth = FMath::Max(GetTickSpaceGeometry().GetLocalSize().X, MinDesiredOriginalNameWidth + MinDesiredNewNameWidth + 20.f);
 
-	RenamePreviewListBox->SetWidthOverride(ActualWidth);
-	
 	RenamePreviewListHeaderRow->SetColumnWidth(
-		FAdvancedRenamerPreviewListItem::OriginalNameColumnName, 
+		OriginalNameColumnName, 
 		(MinDesiredOriginalNameWidth + 10.f) / (MinDesiredOriginalNameWidth + MinDesiredNewNameWidth + 20.f) * ActualWidth
 	);
 
 	RenamePreviewListHeaderRow->SetColumnWidth(
-		FAdvancedRenamerPreviewListItem::NewNameColumnName, 
+		NewNameColumnName, 
 		(MinDesiredNewNameWidth + 10.f) / (MinDesiredOriginalNameWidth + MinDesiredNewNameWidth + 20.f) * ActualWidth
 	);
-
-	RenamePreviewListHeaderRow->RefreshColumns();
 }
 
 void SAdvancedRenamerPanel::RemoveSelectedObjects()
@@ -1451,7 +834,7 @@ void SAdvancedRenamerPanel::RemoveSelectedObjects()
 		return;
 	}
 
-	TArray<FObjectRenamePreviewListItemPtr> SelectedItems = RenamePreviewList->GetSelectedItems();
+	TArray<TSharedPtr<FAdvancedRenamerPreview>> SelectedItems = RenamePreviewList->GetSelectedItems();
 
 	if (SelectedItems.Num() == 0)
 	{
@@ -1460,53 +843,31 @@ void SAdvancedRenamerPanel::RemoveSelectedObjects()
 
 	bool bMadeChange = false;
 
-	for (int32 Index = 0; Index < ListData.Num(); ++Index)
+	for (int32 SelectedIndex = 0; SelectedIndex < SelectedItems.Num(); ++SelectedIndex)
 	{
-		if (!ListData[Index].IsValid())
+		if (!SelectedItems[SelectedIndex].IsValid())
 		{
-			RemoveIndex(Index);
-			bMadeChange = true;
-			--Index;
 			continue;
 		}
 
-		int32 MatchIdx = INDEX_NONE;
-
-		for (int32 SelectedIdx = 0; SelectedIdx < SelectedItems.Num(); ++SelectedIdx)
-		{
-			if (!SelectedItems[SelectedIdx].IsValid())
-			{
-				continue;
-			}
-
-			if (SelectedItems[SelectedIdx]->Hash != ListData[Index]->Hash)
-			{
-				continue;
-			}
-
-			MatchIdx = SelectedIdx;
-			break;
-		}
+		const int32 PreviewIndex = Renamer->FindHash(SelectedItems[SelectedIndex]->Hash);
 
-		if (MatchIdx == INDEX_NONE)
+		if (PreviewIndex == INDEX_NONE)
 		{
 			continue;
 		}
 
-		RemoveIndex(Index);
-		--Index;
-		SelectedItems.RemoveAt(MatchIdx);
-		bMadeChange = true;
-
-		if (SelectedItems.Num() == 0)
+		if (Renamer->RemoveIndex(PreviewIndex))
 		{
-			break;
+			bMadeChange = true;
 		}
 	}
 
+	SelectedItems.Empty();
+
 	if (bMadeChange)
 	{
-		if (ListData.Num() == 0)
+		if (Renamer->Num() == 0)
 		{
 			if (CloseWindow())
 			{
@@ -1518,54 +879,66 @@ void SAdvancedRenamerPanel::RemoveSelectedObjects()
 	}
 }
 
-void SAdvancedRenamerPanel::Tick(const FGeometry& AllottedGeometry, const double InCurrentTime, const float InDeltaTime)
+void SAdvancedRenamerPanel::Tick(const FGeometry& InAllottedGeometry, const double InCurrentTime, const float InDeltaTime)
 {
-	SCompoundWidget::Tick(AllottedGeometry, InCurrentTime, InDeltaTime);
+	SCompoundWidget::Tick(InAllottedGeometry, InCurrentTime, InDeltaTime);
 
 	if (ListLastUpdateTime == 0)
 	{
 		ListLastUpdateTime = InCurrentTime;
 	}
-	else if (bListNeedsUpdate && InCurrentTime >= (ListLastUpdateTime + MinUpdateFrequency))
+	else if (Renamer->IsDirty() && InCurrentTime >= (ListLastUpdateTime + MinUpdateFrequency))
 	{
 		RefreshListView(InCurrentTime);
 	}
+	else
+	{
+		UpdateRequiredListWidth();
+	}
 }
 
-void SAdvancedRenamerPanel::OnBaseNameCheckBoxChanged(ECheckBoxState NewState)
+void SAdvancedRenamerPanel::OnBaseNameChanged(const FText& InNewText)
 {
-	bBaseName = NewState == ECheckBoxState::Checked;
-
-	UpdateEnables();
+	Renamer->GetOptions().BaseName = InNewText.ToString();
+	Renamer->MarkDirty();
 }
 
-void SAdvancedRenamerPanel::OnBaseNameChanged(const FText& NewText)
+void SAdvancedRenamerPanel::OnPrefixChanged(const FText& InNewText)
 {
-	RequestListViewRefresh();
+	Renamer->GetOptions().AddPrefix = InNewText.ToString();
+	Renamer->MarkDirty();
 }
 
-void SAdvancedRenamerPanel::OnPrefixCheckBoxChanged(ECheckBoxState NewState)
+ECheckBoxState SAdvancedRenamerPanel::IsPrefixRemoveChecked() const
 {
-	bPrefix = NewState == ECheckBoxState::Checked;
-
-	UpdateEnables();
+	return bRemovePrefixSeparator ? ECheckBoxState::Checked : ECheckBoxState::Unchecked;
 }
 
-void SAdvancedRenamerPanel::OnPrefixChanged(const FText& NewText)
+void SAdvancedRenamerPanel::OnPrefixRemoveCheckBoxChanged(ECheckBoxState InNewState)
 {
-	RequestListViewRefresh();
-}
+	bRemovePrefixSeparator = InNewState == ECheckBoxState::Checked;
 
-void SAdvancedRenamerPanel::OnPrefixRemoveCheckBoxChanged(ECheckBoxState NewState)
-{
-	bPrefixRemove = NewState == ECheckBoxState::Checked;
+	if (bRemovePrefixSeparator)
+	{
+		bRemovePrefixNumChars = false;
+	}
 
-	if (bPrefixRemove)
+	if (bRemovePrefixSeparator && PrefixSeparatorTextBox.IsValid())
+	{
+		Renamer->GetOptions().RemovePrefixSeparator = PrefixSeparatorTextBox->GetText().ToString();
+		Renamer->GetOptions().RemovePrefixCharacterCount = 0;
+	}
+	else
 	{
-		bPrefixRemoveCharacters = false;
+		Renamer->GetOptions().RemovePrefixSeparator = "";
 	}
 
-	UpdateEnables();
+	Renamer->MarkDirty();
+}
+
+bool SAdvancedRenamerPanel::IsPrefixRemoveSeparatorEnabled() const
+{
+	return bRemovePrefixSeparator;
 }
 
 bool SAdvancedRenamerPanel::OnPrefixSeparatorVerifyTextChanged(const FText& InText, FText& OutErrorText) const
@@ -1579,50 +952,86 @@ bool SAdvancedRenamerPanel::OnPrefixSeparatorVerifyTextChanged(const FText& InTe
 	return true;
 }
 
-void SAdvancedRenamerPanel::OnPrefixSeparatorChanged(const FText& NewText)
+void SAdvancedRenamerPanel::OnPrefixSeparatorChanged(const FText& InNewText)
+{
+	Renamer->GetOptions().RemovePrefixSeparator = InNewText.ToString();
+	Renamer->MarkDirty();
+}
+
+ECheckBoxState SAdvancedRenamerPanel::IsPrefixRemoveCharactersChecked() const
 {
-	RequestListViewRefresh();
+	return bRemovePrefixNumChars ? ECheckBoxState::Checked : ECheckBoxState::Unchecked;
 }
 
-void SAdvancedRenamerPanel::OnPrefixRemoveCharactersCheckBoxChanged(ECheckBoxState NewState)
+void SAdvancedRenamerPanel::OnPrefixRemoveCharactersCheckBoxChanged(ECheckBoxState InNewState)
 {
-	bPrefixRemoveCharacters = NewState == ECheckBoxState::Checked;
+	bRemovePrefixNumChars = InNewState == ECheckBoxState::Checked;
+
+	if (bRemovePrefixNumChars)
+	{
+		bRemovePrefixSeparator = false;
+	}
 
-	if (bPrefixRemoveCharacters)
+	if (bRemovePrefixNumChars && PrefixRemoveCharactersSpinBox.IsValid())
 	{
-		bPrefixRemove = false;
+		Renamer->GetOptions().RemovePrefixCharacterCount = PrefixRemoveCharactersSpinBox->GetValue();
+		Renamer->GetOptions().RemovePrefixSeparator = "";
+	}
+	else
+	{
+		Renamer->GetOptions().RemovePrefixCharacterCount = 0;
 	}
 
-	UpdateEnables();
+	Renamer->MarkDirty();
 }
 
-void SAdvancedRenamerPanel::OnPrefixRemoveCharactersChanged(uint8 NewValue)
+bool SAdvancedRenamerPanel::IsPrefixRemoveNumCharsEnabled() const
 {
-	RequestListViewRefresh();
+	return bRemovePrefixNumChars;
 }
 
-void SAdvancedRenamerPanel::OnSuffixCheckBoxChanged(ECheckBoxState NewState)
+void SAdvancedRenamerPanel::OnPrefixRemoveCharactersChanged(uint8 InNewValue)
 {
-	bSuffix = NewState == ECheckBoxState::Checked;
+	Renamer->GetOptions().RemovePrefixCharacterCount = InNewValue;
+	Renamer->MarkDirty();
+}
 
-	UpdateEnables();
+void SAdvancedRenamerPanel::OnSuffixChanged(const FText& InNewText)
+{
+	Renamer->GetOptions().AddSuffix = InNewText.ToString();
+	Renamer->MarkDirty();
 }
 
-void SAdvancedRenamerPanel::OnSuffixChanged(const FText& NewText)
+ECheckBoxState SAdvancedRenamerPanel::IsSuffixRemoveChecked() const
 {
-	RequestListViewRefresh();
+	return bRemoveSuffixSeparator ? ECheckBoxState::Checked : ECheckBoxState::Unchecked;
 }
 
-void SAdvancedRenamerPanel::OnSuffixRemoveCheckBoxChanged(ECheckBoxState NewState)
+void SAdvancedRenamerPanel::OnSuffixRemoveCheckBoxChanged(ECheckBoxState InNewState)
 {
-	bSuffixRemove = NewState == ECheckBoxState::Checked;
+	bRemoveSuffixSeparator = InNewState == ECheckBoxState::Checked;
+
+	if (bRemoveSuffixSeparator)
+	{
+		bRemoveSuffixNumChars = false;
+	}
 
-	if (bSuffixRemove)
+	if (bRemoveSuffixSeparator && SuffixSeparatorTextBox.IsValid())
+	{
+		Renamer->GetOptions().RemoveSuffixSeparator = SuffixSeparatorTextBox->GetText().ToString();
+		Renamer->GetOptions().RemoveSuffixCharacterCount = 0;
+	}
+	else
 	{
-		bSuffixRemoveCharacters = false;
+		Renamer->GetOptions().RemoveSuffixSeparator = "";
 	}
 
-	UpdateEnables();
+	Renamer->MarkDirty();
+}
+
+bool SAdvancedRenamerPanel::IsSuffixRemoveSeparatorEnabled() const
+{
+	return bRemoveSuffixSeparator;
 }
 
 bool SAdvancedRenamerPanel::OnSuffixSeparatorVerifyTextChanged(const FText& InText, FText& OutErrorText) const
@@ -1636,107 +1045,190 @@ bool SAdvancedRenamerPanel::OnSuffixSeparatorVerifyTextChanged(const FText& InTe
 	return true;
 }
 
-void SAdvancedRenamerPanel::OnSuffixSeparatorChanged(const FText& NewText)
+void SAdvancedRenamerPanel::OnSuffixSeparatorChanged(const FText& InNewText)
+{
+	Renamer->GetOptions().RemoveSuffixSeparator = InNewText.ToString();
+	Renamer->MarkDirty();
+}
+
+ECheckBoxState SAdvancedRenamerPanel::IsSuffixRemoveCharactersChecked() const
+{
+	return bRemoveSuffixNumChars ? ECheckBoxState::Checked : ECheckBoxState::Unchecked;
+}
+
+bool SAdvancedRenamerPanel::IsSuffixRemoveNumCharsEnabled() const
 {
-	RequestListViewRefresh();
+	return bRemoveSuffixNumChars;
 }
 
-void SAdvancedRenamerPanel::OnSuffixRemoveCharactersCheckBoxChanged(ECheckBoxState NewState)
+void SAdvancedRenamerPanel::OnSuffixRemoveCharactersCheckBoxChanged(ECheckBoxState InNewState)
 {
-	bSuffixRemoveCharacters = NewState == ECheckBoxState::Checked;
+	bRemoveSuffixNumChars = InNewState == ECheckBoxState::Checked;
 
-	if (bSuffixRemoveCharacters)
+	if (bRemoveSuffixNumChars)
 	{
-		bSuffixRemove = false;
+		bRemoveSuffixSeparator = false;
 	}
 
-	UpdateEnables();
+	if (bRemoveSuffixNumChars && SuffixRemoveCharactersSpinBox.IsValid())
+	{
+		Renamer->GetOptions().RemoveSuffixCharacterCount = SuffixRemoveCharactersSpinBox->GetValue();
+		Renamer->GetOptions().RemoveSuffixSeparator = "";
+	}
+	else
+	{
+		Renamer->GetOptions().RemoveSuffixCharacterCount = 0;
+	}
+
+	Renamer->MarkDirty();
+}
+
+void SAdvancedRenamerPanel::OnSuffixRemoveCharactersChanged(uint8 InNewValue)
+{
+	Renamer->GetOptions().RemoveSuffixCharacterCount = InNewValue;
+	Renamer->MarkDirty();
 }
 
-void SAdvancedRenamerPanel::OnSuffixRemoveCharactersChanged(uint8 NewValue)
+ECheckBoxState SAdvancedRenamerPanel::IsSuffixRemoveNumberChecked() const
 {
-	RequestListViewRefresh();
+	return Renamer->GetOptions().bRemoveSuffixNumber ? ECheckBoxState::Checked : ECheckBoxState::Unchecked;
 }
 
-void SAdvancedRenamerPanel::OnSuffixRemoveNumberCheckBoxChanged(ECheckBoxState NewState)
+void SAdvancedRenamerPanel::OnSuffixRemoveNumberCheckBoxChanged(ECheckBoxState InNewState)
 {
-	bSuffixRemoveNumber = NewState == ECheckBoxState::Checked;
+	Renamer->GetOptions().bRemoveSuffixNumber = InNewState == ECheckBoxState::Checked;
 
-	UpdateEnables();
+	if (!Renamer->GetOptions().bRemoveSuffixNumber)
+	{
+		Renamer->GetOptions().bAddSuffixNumber = false;
+	}
+
+	Renamer->MarkDirty();
 }
 
-void SAdvancedRenamerPanel::OnSuffixNumberCheckBoxChanged(ECheckBoxState NewState)
+ECheckBoxState SAdvancedRenamerPanel::IsSuffixNumberChecked() const
 {
-	bSuffixNumber = NewState == ECheckBoxState::Checked;
+	return Renamer->GetOptions().bAddSuffixNumber ? ECheckBoxState::Checked : ECheckBoxState::Unchecked;
+}
 
-	if (bSuffixNumber)
+bool SAdvancedRenamerPanel::IsSuffixRemoveNumberCheckBoxEnabled() const
+{
+	return !Renamer->GetOptions().bAddSuffixNumber;
+}
+
+void SAdvancedRenamerPanel::OnSuffixNumberCheckBoxChanged(ECheckBoxState InNewState)
+{
+	Renamer->GetOptions().bAddSuffixNumber = InNewState == ECheckBoxState::Checked;
+
+	if (Renamer->GetOptions().bAddSuffixNumber)
 	{
-		bSuffixRemoveNumber = true;
+		Renamer->GetOptions().bRemoveSuffixNumber = true;
 	}
 
-	UpdateEnables();
+	Renamer->MarkDirty();
 }
 
-void SAdvancedRenamerPanel::OnSuffixNumberStartChanged(int32 NewValue)
+void SAdvancedRenamerPanel::OnSuffixNumberStartChanged(int32 InNewValue)
 {
-	RequestListViewRefresh();
+	Renamer->GetOptions().AddSuffixNumberStart = InNewValue;
+	Renamer->MarkDirty();
 }
 
-void SAdvancedRenamerPanel::OnSuffixNumberStepChanged(int32 NewValue)
+void SAdvancedRenamerPanel::OnSuffixNumberStepChanged(int32 InNewValue)
 {
-	RequestListViewRefresh();
+	Renamer->GetOptions().AddSuffixNumberStep = InNewValue;
+	Renamer->MarkDirty();
 }
 
-void SAdvancedRenamerPanel::OnSearchReplacePlainTextCheckBoxChanged(ECheckBoxState NewState)
+ECheckBoxState SAdvancedRenamerPanel::IsSearchReplacePlainTextChecked() const
 {
-	bSearchReplacePlainText = NewState == ECheckBoxState::Checked;
+	return Renamer->GetOptions().SearchAndReplaceType == EAdvancedRenamerSeachAndReplaceType::PlainText ? ECheckBoxState::Checked : ECheckBoxState::Unchecked;
+}
 
-	if (bSearchReplacePlainText)
+void SAdvancedRenamerPanel::OnSearchReplacePlainTextCheckBoxChanged(ECheckBoxState InNewState)
+{
+	if (InNewState != ECheckBoxState::Checked)
 	{
-		bSearchReplaceRegex = false;
+		return;
 	}
 
-	UpdateEnables();
+	Renamer->GetOptions().SearchAndReplaceType = EAdvancedRenamerSeachAndReplaceType::PlainText;
+	Renamer->MarkDirty();
 }
 
-void SAdvancedRenamerPanel::OnSearchReplaceRegexCheckBoxChanged(ECheckBoxState NewState)
+ECheckBoxState SAdvancedRenamerPanel::IsSearchReplaceRegexChecked() const
 {
-	bSearchReplaceRegex = NewState == ECheckBoxState::Checked;
+	return Renamer->GetOptions().SearchAndReplaceType == EAdvancedRenamerSeachAndReplaceType::RegularExpression ? ECheckBoxState::Checked : ECheckBoxState::Unchecked;
+}
 
-	if (bSearchReplaceRegex)
+void SAdvancedRenamerPanel::OnSearchReplaceRegexCheckBoxChanged(ECheckBoxState InNewState)
+{
+	if (InNewState != ECheckBoxState::Checked)
 	{
-		bSearchReplacePlainText = false;
+		return;
 	}
 
-	UpdateEnables();
+	Renamer->GetOptions().SearchAndReplaceType = EAdvancedRenamerSeachAndReplaceType::RegularExpression;
+	Renamer->MarkDirty();
+}
+
+ECheckBoxState SAdvancedRenamerPanel::IsSearchReplaceIgnoreCaseChecked() const
+{
+	return Renamer->GetOptions().SearchAndReplaceCase == ESearchCase::IgnoreCase ? ECheckBoxState::Checked : ECheckBoxState::Unchecked;
 }
 
-void SAdvancedRenamerPanel::OnSearchReplaceIgnoreCaseCheckBoxChanged(ECheckBoxState NewState)
+void SAdvancedRenamerPanel::OnSearchReplaceIgnoreCaseCheckBoxChanged(ECheckBoxState InNewState)
 {
-	bSearchReplaceIgnoreCase = NewState == ECheckBoxState::Checked;
+	switch (InNewState)
+	{
+		case ECheckBoxState::Checked:
+			Renamer->GetOptions().SearchAndReplaceCase = ESearchCase::IgnoreCase;
+			break;
+
+		case ECheckBoxState::Unchecked:
+			Renamer->GetOptions().SearchAndReplaceCase = ESearchCase::CaseSensitive;
+			break;
 
-	UpdateEnables();
+		default:
+			return;
+	}
+
+	Renamer->MarkDirty();
 }
 
-void SAdvancedRenamerPanel::OnSearchReplaceSearchTextChanged(const FText& NewText)
+void SAdvancedRenamerPanel::OnSearchReplaceSearchTextChanged(const FText& InNewText)
 {
-	RequestListViewRefresh();
+	Renamer->GetOptions().SearchAndReplaceFromText = InNewText.ToString();
+
+	if (Renamer->GetOptions().SearchAndReplaceType == EAdvancedRenamerSeachAndReplaceType::None)
+	{
+		Renamer->GetOptions().SearchAndReplaceType = EAdvancedRenamerSeachAndReplaceType::PlainText;
+	}
+
+	Renamer->MarkDirty();
 }
 
-void SAdvancedRenamerPanel::OnSearchReplaceReplaceTextChanged(const FText& NewText)
+void SAdvancedRenamerPanel::OnSearchReplaceReplaceTextChanged(const FText& InNewText)
 {
-	RequestListViewRefresh();
+	Renamer->GetOptions().SearchAndReplaceToText = InNewText.ToString();
+
+	if (Renamer->GetOptions().SearchAndReplaceType == EAdvancedRenamerSeachAndReplaceType::None)
+	{
+		Renamer->GetOptions().SearchAndReplaceType = EAdvancedRenamerSeachAndReplaceType::PlainText;
+	}
+
+	Renamer->MarkDirty();
 }
 
-TSharedRef<ITableRow> SAdvancedRenamerPanel::OnGenerateRowForList(FObjectRenamePreviewListItemPtr Item, 
-	const TSharedRef<STableViewBase>& OwnerTable)
+TSharedRef<ITableRow> SAdvancedRenamerPanel::OnGenerateRowForList(TSharedPtr<FAdvancedRenamerPreview> InItem, 
+	const TSharedRef<STableViewBase>& InOwnerTable)
 {
-	return SNew(SAdvancedRenamerPreviewListRow, SharedThis(this), OwnerTable, Item);
+	return SNew(SAdvancedRenamerPreviewListRow, SharedThis(this), InOwnerTable, InItem);
 }
 
-FReply SAdvancedRenamerPanel::OnListViewKeyDown(const FGeometry& MyGeometry, const FKeyEvent& KeyEvent)
+FReply SAdvancedRenamerPanel::OnListViewKeyDown(const FGeometry& InGeometry, const FKeyEvent& InKeyEvent)
 {
-	if (CommandList.IsValid() && CommandList->ProcessCommandBindings(KeyEvent))
+	if (CommandList.IsValid() && CommandList->ProcessCommandBindings(InKeyEvent))
 	{
 		return FReply::Handled();
 	}
@@ -1767,9 +1259,16 @@ TSharedPtr<SWidget> SAdvancedRenamerPanel::GenerateListViewContextMenu()
 	return MenuBuilder.MakeWidget();
 }
 
+bool SAdvancedRenamerPanel::IsApplyButtonEnabled() const
+{
+	return !Renamer->IsDirty() && Renamer->HasRenames();
+}
+
 FReply SAdvancedRenamerPanel::OnApplyButtonClicked()
 {
-	if (RenameObjects())
+	FScopedTransaction Transaction(LOCTEXT("AdvancedRenamerRename", "Advanced Renamer Rename"));
+
+	if (Renamer->Execute())
 	{
 		CloseWindow();
 	}
@@ -1777,46 +1276,40 @@ FReply SAdvancedRenamerPanel::OnApplyButtonClicked()
 	return FReply::Handled();
 }
 
-int32 SAdvancedRenamerPanel::Num() const
+FVector2D SAdvancedRenamerPanel::GetRightPaneSize() const
 {
-	return SharedProvider->Num();
-}
+	using namespace UE::AdvancedRenamer::Private;
 
-bool SAdvancedRenamerPanel::IsValidIndex(int32 Index) const
-{
-	return SharedProvider->IsValidIndex(Index);
-}
+	const FVector2f GeoSize = GetTickSpaceGeometry().GetLocalSize();
 
-uint32 SAdvancedRenamerPanel::GetHash(int32 Index) const
-{
-	return SharedProvider->GetHash(Index);
+	return FVector2D(
+		FMath::Max(RightBlockSize.X, GeoSize.X - 319.0),
+		RightBlockSize.Y
+	);
 }
 
-FString SAdvancedRenamerPanel::GetOriginalName(int32 Index) const
+FVector2D SAdvancedRenamerPanel::GetListViewsize() const
 {
-	return SharedProvider->GetOriginalName(Index);
-}
+	using namespace UE::AdvancedRenamer::Private;
 
-bool SAdvancedRenamerPanel::RemoveIndex(int32 Index)
-{
-	// Can fail during construction when indices that aren't renameable are removed from the provider before
-	// they are added to ListData.
-	if (ListData.IsValidIndex(Index))
-	{
-		ListData.RemoveAt(Index);
-	}
+	const FVector2f GeoSize = GetTickSpaceGeometry().GetLocalSize();
 
-	return SharedProvider->RemoveIndex(Index);
+	return FVector2D(
+		FMath::Max(ListViewSize.X, GeoSize.X - 323.0),
+		ListViewSize.Y
+	);
 }
 
-bool SAdvancedRenamerPanel::CanRename(int32 Index) const
+FVector2D SAdvancedRenamerPanel::GetApplyButtonSize() const
 {
-	return SharedProvider->CanRename(Index);
-}
+	using namespace UE::AdvancedRenamer::Private;
 
-bool SAdvancedRenamerPanel::ExecuteRename(int32 Index, const FString& NewName)
-{
-	return SharedProvider->ExecuteRename(Index, NewName);
+	const FVector2f GeoSize = GetTickSpaceGeometry().GetLocalSize();
+
+	return FVector2D(
+		FMath::Max(ListViewSize.X, GeoSize.X - 323.0),
+		ApplyButtonHeight
+	);
 }
 
 #undef LOCTEXT_NAMESPACE
diff --git a/Engine/Plugins/Experimental/AdvancedRenamer/Source/AdvancedRenamer/Private/Slate/SAdvancedRenamerPanel.h b/Engine/Plugins/Experimental/AdvancedRenamer/Source/AdvancedRenamer/Private/Slate/SAdvancedRenamerPanel.h
--- a/Engine/Plugins/Experimental/AdvancedRenamer/Source/AdvancedRenamer/Private/Slate/SAdvancedRenamerPanel.h
+++ b/Engine/Plugins/Experimental/AdvancedRenamer/Source/AdvancedRenamer/Private/Slate/SAdvancedRenamerPanel.h
@@ -2,17 +2,14 @@
 
 #pragma once
 
-#include "Providers/IAdvancedRenamerProvider.h"
 #include "Styling/SlateTypes.h"
 #include "UObject/StrongObjectPtr.h"
 #include "Widgets/Views/SListView.h"
 #include "Widgets/Views/STableRow.h"
 #include "Widgets/Views/STableViewBase.h"
 
-class FRegexPattern;
 class FUICommandList;
-class IAdvancedRenamerProvider;
-class SAdvancedRenamerPanel;
+class IAdvancedRenamer;
 class SBox;
 class SButton;
 class SCanvas;
@@ -21,99 +18,40 @@ class SEditableTextBox;
 class SHeaderRow;
 class SMultiLineEditableTextBox; 
 class UObject;
+struct FAdvancedRenamerPreview;
 template<typename NumericType> class SSpinBox;
 
-struct FAdvancedRenamerPreviewListItem
-{
-	static FName OriginalNameColumnName;
-	static FName NewNameColumnName;
-
-	FAdvancedRenamerPreviewListItem(int32 InHash, const FString InOriginalName)
-		: Hash(InHash)
-		, OriginalName(InOriginalName)
-		, NewName(FString(""))
-	{
-	}
-
-	int32 Hash;
-	const FString OriginalName;
-	mutable FString NewName;
-};
-
-using FObjectRenamePreviewListItemPtr = TSharedPtr<FAdvancedRenamerPreviewListItem, ESPMode::ThreadSafe>;
-using FObjectRenamePreviewListItemWeakPtr = TWeakPtr<FAdvancedRenamerPreviewListItem, ESPMode::ThreadSafe>;
-
-class SAdvancedRenamerPreviewListRow : public SMultiColumnTableRow<FObjectRenamePreviewListItemPtr>
-{
-public:
-	SLATE_BEGIN_ARGS(SAdvancedRenamerPreviewListRow) {}
-	SLATE_END_ARGS()
-
-	void Construct(const FArguments& InArgs, TSharedPtr<SAdvancedRenamerPanel> InRenamePanel, 
-		const TSharedRef<STableViewBase>& InOwnerTableView, FObjectRenamePreviewListItemPtr InRowItem);
-
-	virtual TSharedRef<SWidget> GenerateWidgetForColumn(const FName& ColumnName) override;
-
-protected:
-
-	TWeakPtr<SAdvancedRenamerPanel> RenamePanel;
-	FObjectRenamePreviewListItemWeakPtr RowItem;
-
-};
-
-/**
- * Implements its own provider interface so it can avoid long Execute_ lines and handle
- * the 2 different types of provider (SharedPtr and UObject.)
- */
-class SAdvancedRenamerPanel : public SCompoundWidget, private IAdvancedRenamerProvider
+class SAdvancedRenamerPanel : public SCompoundWidget
 {
 	friend class SAdvancedRenamerPreviewListRow;
 
 public:
 	SLATE_BEGIN_ARGS(SAdvancedRenamerPanel) {}
-		SLATE_ARGUMENT(TSharedPtr<IAdvancedRenamerProvider>, SharedProvider)
 	SLATE_END_ARGS()
 
-	void Construct(const FArguments& InArgs);
+	void Construct(const FArguments& InArgs, const TSharedRef<IAdvancedRenamer>& InRenamer);
 
 protected:
-	static const inline double MinUpdateFrequency = 0.1;
+	static constexpr double MinUpdateFrequency = 0.1;
 
-	TSharedPtr<IAdvancedRenamerProvider> SharedProvider;
-	TArray<TSharedPtr<FAdvancedRenamerPreviewListItem>> ListData;
+	TSharedPtr<IAdvancedRenamer> Renamer;
 	TSharedPtr<FUICommandList> CommandList;
 
-	bool bValidNames = false;
-	bool bListNeedsUpdate = false;
-	bool bBaseName = false;
-	bool bPrefix = false;
-	bool bPrefixRemove = false;
-	bool bPrefixRemoveCharacters = false;
-	bool bSuffix = false;
-	bool bSuffixRemove = false;
-	bool bSuffixRemoveCharacters = false;
-	bool bSuffixRemoveNumber = false;
-	bool bSuffixNumber = false;
-	bool bSearchReplacePlainText = false;
-	bool bSearchReplaceRegex = false;
-	bool bSearchReplaceIgnoreCase = true;
-
 	double ListLastUpdateTime = 0;
 	float MinDesiredOriginalNameWidth = 0.f;
 	float MinDesiredNewNameWidth = 0.f;
 
-	uint8 PrefixRemoveCharacterCount = 0;
-	uint8 SuffixRemoveCharacterCount = 0;
+	bool bRemovePrefixSeparator;
+	bool bRemovePrefixNumChars;
+	bool bRemoveSuffixSeparator;
+	bool bRemoveSuffixNumChars;
 
-	TSharedPtr<SCheckBox> BaseNameCheckBox;
 	TSharedPtr<SEditableTextBox> BaseNameTextBox;
-	TSharedPtr<SCheckBox> PrefixCheckBox;
 	TSharedPtr<SEditableTextBox> PrefixTextBox;
 	TSharedPtr<SCheckBox> PrefixRemoveCheckBox;
 	TSharedPtr<SEditableTextBox> PrefixSeparatorTextBox;
 	TSharedPtr<SCheckBox> PrefixRemoveCharactersCheckBox;
 	TSharedPtr<SSpinBox<uint8>> PrefixRemoveCharactersSpinBox;
-	TSharedPtr<SCheckBox> SuffixCheckBox;
 	TSharedPtr<SEditableTextBox> SuffixTextBox;
 	TSharedPtr<SCheckBox> SuffixRemoveCheckBox;
 	TSharedPtr<SEditableTextBox> SuffixSeparatorTextBox;
@@ -130,112 +68,95 @@ protected:
 	TSharedPtr<SMultiLineEditableTextBox> SearchReplaceReplaceTextBox;
 	TSharedPtr<SBox> RenamePreviewListBox;
 	TSharedPtr<SHeaderRow> RenamePreviewListHeaderRow;
-	TSharedPtr<SListView<FObjectRenamePreviewListItemPtr>> RenamePreviewList;
+	TSharedPtr<SListView<TSharedPtr<FAdvancedRenamerPreview>>> RenamePreviewList;
 	TSharedPtr<SButton> ApplyButton;
 
-	void CreateLeftPane(TSharedRef<SCanvas> Canvas);
+	void CreateLeftPane(const TSharedRef<SCanvas>& InCanvas);
 	TSharedRef<SWidget> CreateBaseName();
 	TSharedRef<SWidget> CreatePrefix();
 	TSharedRef<SWidget> CreateSuffix();
 	TSharedRef<SWidget> CreateSearchAndReplace();
 
-	void CreateRightPane(TSharedRef<SCanvas> Canvas);
+	void CreateRightPane(const TSharedRef<SCanvas>& InCanvas);
 	TSharedRef<SWidget> CreateRenamePreview();
 
-	bool RenameObjects();
-
 	bool CloseWindow();
 
-	FString CreateNewName(int32 Index) const;
-	FString ApplyRename(const FString& OriginalName, int32 Index) const;
-	FString ApplyBaseName(const FString& OriginalName) const;
-	FString ApplyPrefix(const FString& OriginalName) const;
-	FString ApplySuffix(const FString& OriginalName, int32 Index) const;
-	FString ApplySearchPlainText(const FString& OriginalName) const;
-	FString ApplySearchReplaceRegex(const FString& OriginalName) const;
-
-	FString RegexReplace(const FString& OriginalString, const FRegexPattern Pattern, const FString& ReplaceString) const;
-
-	void UpdateEnables();
-	void RequestListViewRefresh();
 	void RefreshListView(const double InCurrentTime);
 	void UpdateRequiredListWidth();
 
 	void RemoveSelectedObjects();
 
-	virtual void Tick(const FGeometry& AllottedGeometry, const double InCurrentTime, const float InDeltaTime) override;
-
-	ECheckBoxState IsBaseNameChecked() const { return bBaseName ? ECheckBoxState::Checked : ECheckBoxState::Unchecked; }
-	void OnBaseNameCheckBoxChanged(ECheckBoxState NewState);
+	virtual void Tick(const FGeometry& InAllottedGeometry, const double InCurrentTime, const float InDeltaTime) override;
 
-	void OnBaseNameChanged(const FText& NewText);
+	void OnBaseNameChanged(const FText& InNewText);
 
-	ECheckBoxState IsPrefixChecked() const { return bPrefix ? ECheckBoxState::Checked : ECheckBoxState::Unchecked; }
-	void OnPrefixCheckBoxChanged(ECheckBoxState NewState);
+	void OnPrefixChanged(const FText& InNewText);
 
-	void OnPrefixChanged(const FText& NewText);
-
-	ECheckBoxState IsPrefixRemoveChecked() const { return bPrefixRemove ? ECheckBoxState::Checked : ECheckBoxState::Unchecked; }
-	void OnPrefixRemoveCheckBoxChanged(ECheckBoxState NewState);
+	ECheckBoxState IsPrefixRemoveChecked() const;
+	void OnPrefixRemoveCheckBoxChanged(ECheckBoxState InNewState);
 
+	bool IsPrefixRemoveSeparatorEnabled() const;
 	bool OnPrefixSeparatorVerifyTextChanged(const FText& InText, FText& OutErrorText) const;
-	void OnPrefixSeparatorChanged(const FText& NewText);
-
-	ECheckBoxState IsPrefixRemoveCharactersChecked() const { return bPrefixRemoveCharacters ? ECheckBoxState::Checked : ECheckBoxState::Unchecked; }
-	void OnPrefixRemoveCharactersCheckBoxChanged(ECheckBoxState NewState);
+	void OnPrefixSeparatorChanged(const FText& InNewText);
 
-	void OnPrefixRemoveCharactersChanged(uint8 NewValue);
+	ECheckBoxState IsPrefixRemoveCharactersChecked() const;
+	void OnPrefixRemoveCharactersCheckBoxChanged(ECheckBoxState InNewState);
 
-	ECheckBoxState IsSuffixChecked() const { return bSuffix ? ECheckBoxState::Checked : ECheckBoxState::Unchecked; }
-	void OnSuffixCheckBoxChanged(ECheckBoxState NewState);
+	bool IsPrefixRemoveNumCharsEnabled() const;
+	void OnPrefixRemoveCharactersChanged(uint8 InNewValue);
 
-	void OnSuffixChanged(const FText& NewText);
+	void OnSuffixChanged(const FText& InNewText);
 
-	ECheckBoxState IsSuffixRemoveChecked() const { return bSuffixRemove ? ECheckBoxState::Checked : ECheckBoxState::Unchecked; }
-	void OnSuffixRemoveCheckBoxChanged(ECheckBoxState NewState);
+	ECheckBoxState IsSuffixRemoveChecked() const;
+	void OnSuffixRemoveCheckBoxChanged(ECheckBoxState InNewState);
 
+	bool IsSuffixRemoveSeparatorEnabled() const;
 	bool OnSuffixSeparatorVerifyTextChanged(const FText& InText, FText& OutErrorText) const;
-	void OnSuffixSeparatorChanged(const FText& NewText);
+	void OnSuffixSeparatorChanged(const FText& InNewText);
+
+	ECheckBoxState IsSuffixRemoveCharactersChecked() const;
+	void OnSuffixRemoveCharactersCheckBoxChanged(ECheckBoxState InNewState);
 
-	ECheckBoxState IsSuffixRemoveCharactersChecked() const { return bSuffixRemoveCharacters ? ECheckBoxState::Checked : ECheckBoxState::Unchecked; }
-	void OnSuffixRemoveCharactersCheckBoxChanged(ECheckBoxState NewState);
+	bool IsSuffixRemoveNumCharsEnabled() const;
+	void OnSuffixRemoveCharactersChanged(uint8 InNewValue);
 
-	void OnSuffixRemoveCharactersChanged(uint8 NewValue);
+	ECheckBoxState IsSuffixRemoveNumberChecked() const;
+	void OnSuffixRemoveNumberCheckBoxChanged(ECheckBoxState InNewState);
 
-	ECheckBoxState IsSuffixRemoveNumberChecked() const { return bSuffixRemoveNumber ? ECheckBoxState::Checked : ECheckBoxState::Unchecked; }
-	void OnSuffixRemoveNumberCheckBoxChanged(ECheckBoxState NewState);
+	ECheckBoxState IsSuffixNumberChecked() const;
+	bool IsSuffixRemoveNumberCheckBoxEnabled() const;
+	void OnSuffixNumberCheckBoxChanged(ECheckBoxState InNewState);
 
-	ECheckBoxState IsSuffixNumberChecked() const { return bSuffixNumber ? ECheckBoxState::Checked : ECheckBoxState::Unchecked; }
-	void OnSuffixNumberCheckBoxChanged(ECheckBoxState NewState);
+	void OnSuffixNumberStartChanged(int32 InNewValue);
 
-	void OnSuffixNumberStartChanged(int32 NewValue);
+	void OnSuffixNumberStepChanged(int32 InNewValue);
 
-	void OnSuffixNumberStepChanged(int32 NewValue);
+	ECheckBoxState IsSearchReplacePlainTextChecked() const;
+	void OnSearchReplacePlainTextCheckBoxChanged(ECheckBoxState InNewState);
 
-	ECheckBoxState IsSearchReplacePlainTextChecked() const { return bSearchReplacePlainText ? ECheckBoxState::Checked : ECheckBoxState::Unchecked; }
-	void OnSearchReplacePlainTextCheckBoxChanged(ECheckBoxState NewState);
+	ECheckBoxState IsSearchReplaceRegexChecked() const;
+	void OnSearchReplaceRegexCheckBoxChanged(ECheckBoxState InNewState);
 
-	ECheckBoxState IsSearchReplaceRegexChecked() const { return bSearchReplaceRegex ? ECheckBoxState::Checked : ECheckBoxState::Unchecked; }
-	void OnSearchReplaceRegexCheckBoxChanged(ECheckBoxState NewState);
+	ECheckBoxState IsSearchReplaceIgnoreCaseChecked() const;
+	void OnSearchReplaceIgnoreCaseCheckBoxChanged(ECheckBoxState InNewState);
 
-	ECheckBoxState IsSearchReplaceIgnoreCaseChecked() const { return bSearchReplaceIgnoreCase ? ECheckBoxState::Checked : ECheckBoxState::Unchecked; }
-	void OnSearchReplaceIgnoreCaseCheckBoxChanged(ECheckBoxState NewState);
+	void OnSearchReplaceSearchTextChanged(const FText& InNewText);
 
-	void OnSearchReplaceSearchTextChanged(const FText& NewText);
+	void OnSearchReplaceReplaceTextChanged(const FText& InNewText);
 
-	void OnSearchReplaceReplaceTextChanged(const FText& NewText);
+	TSharedRef<ITableRow> OnGenerateRowForList(TSharedPtr<FAdvancedRenamerPreview> InItem, const TSharedRef<STableViewBase>& InOwnerTable);
+
+	FReply OnListViewKeyDown(const FGeometry& InGeometry, const FKeyEvent& InKeyEvent);
 
-	TSharedRef<ITableRow> OnGenerateRowForList(FObjectRenamePreviewListItemPtr Item, const TSharedRef<STableViewBase>& OwnerTable);
-	FReply OnListViewKeyDown(const FGeometry& MyGeometry, const FKeyEvent& KeyEvent);
 	TSharedPtr<SWidget> GenerateListViewContextMenu();
 
+	bool IsApplyButtonEnabled() const;
 	FReply OnApplyButtonClicked();
 
-	virtual int32 Num() const override;
-	virtual bool IsValidIndex(int32 Index) const override;
-	virtual uint32 GetHash(int32 Index) const override;
-	virtual FString GetOriginalName(int32 Index) const override;
-	virtual bool RemoveIndex(int32 Index) override;
-	virtual bool CanRename(int32 Index) const override;
-	virtual bool ExecuteRename(int32 Index, const FString& NewName) override;
+	FVector2D GetRightPaneSize() const;
+
+	FVector2D GetListViewsize() const;
+
+	FVector2D GetApplyButtonSize() const;
 };
diff --git a/Engine/Plugins/Experimental/AdvancedRenamer/Source/AdvancedRenamer/Public/IAdvancedRenamerModule.h b/Engine/Plugins/Experimental/AdvancedRenamer/Source/AdvancedRenamer/Public/IAdvancedRenamerModule.h
--- a/Engine/Plugins/Experimental/AdvancedRenamer/Source/AdvancedRenamer/Public/IAdvancedRenamerModule.h
+++ b/Engine/Plugins/Experimental/AdvancedRenamer/Source/AdvancedRenamer/Public/IAdvancedRenamerModule.h
@@ -7,6 +7,7 @@
 #include "Templates/SharedPointerFwd.h"
 
 class AActor;
+class IAdvancedRenamer;
 class IAdvancedRenamerProvider;
 class IToolkitHost;
 class SWidget;
@@ -28,6 +29,8 @@ public:
 		return FModuleManager::LoadModuleChecked<IAdvancedRenamerModule>(ModuleName);
 	}
 
+	virtual TSharedRef<IAdvancedRenamer> CreateAdvancedRenamer(const TSharedRef<IAdvancedRenamerProvider>& InRenameProvider) = 0;
+
 	virtual void OpenAdvancedRenamer(const TSharedRef<IAdvancedRenamerProvider>& InRenameProvider, const TSharedPtr<IToolkitHost>& InToolkitHost) = 0;
 
 	virtual void OpenAdvancedRenamer(const TSharedRef<IAdvancedRenamerProvider>& InRenameProvider, const TSharedPtr<SWidget>& InParentWidget) = 0;
@@ -36,5 +39,9 @@ public:
 
 	virtual void OpenAdvancedRenamerForActors(const TArray<AActor*>& InActors, const TSharedPtr<SWidget>& InParentWidget) = 0;
 
+	virtual void OpenAdvancedRenamer(const TSharedRef<IAdvancedRenamer>& InRenamer, const TSharedPtr<IToolkitHost>& InToolkitHost) = 0;
+
+	virtual void OpenAdvancedRenamer(const TSharedRef<IAdvancedRenamer>& InRenamer, const TSharedPtr<SWidget>& InParentWidget) = 0;
+
 	virtual TArray<AActor*> GetActorsSharingClassesInWorld(const TArray<AActor*>& InActors) = 0;
 };
diff --git a/Engine/Plugins/Experimental/AdvancedRenamer/Source/AdvancedRenamer/Public/Providers/IAdvancedRenamerProvider.h b/Engine/Plugins/Experimental/AdvancedRenamer/Source/AdvancedRenamer/Public/Providers/IAdvancedRenamerProvider.h
--- a/Engine/Plugins/Experimental/AdvancedRenamer/Source/AdvancedRenamer/Public/Providers/IAdvancedRenamerProvider.h
+++ b/Engine/Plugins/Experimental/AdvancedRenamer/Source/AdvancedRenamer/Public/Providers/IAdvancedRenamerProvider.h
@@ -3,6 +3,7 @@
 #pragma once
 
 #include "HAL/Platform.h"
+#include "Misc/CoreMiscDefines.h"
 
 class FString;
 
@@ -11,10 +12,25 @@ class IAdvancedRenamerProvider
 public:
 	virtual ~IAdvancedRenamerProvider() = default;
 	virtual int32 Num() const = 0;
-	virtual bool IsValidIndex(int32 Index) const = 0;
-	virtual uint32 GetHash(int32 Index) const = 0;
-	virtual FString GetOriginalName(int32 Index) const = 0;
-	virtual bool RemoveIndex(int32 Index) = 0;
-	virtual bool CanRename(int32 Index) const = 0;
-	virtual bool ExecuteRename(int32 Index, const FString& NewName) = 0;
+	virtual bool IsValidIndex(int32 InIndex) const = 0;
+	virtual uint32 GetHash(int32 InIndex) const = 0;
+	virtual FString GetOriginalName(int32 InIndex) const = 0;
+	virtual bool RemoveIndex(int32 InIndex) = 0;
+	virtual bool CanRename(int32 InIndex) const = 0;
+	virtual bool ExecuteRename(int32 InIndex, const FString& InNewName) = 0;
+
+	virtual int32 FindHash(int32 InInHash) const
+	{
+		const int32 Count = Num();
+
+		for (int32 Index = 0; Index < Count; ++Index)
+		{
+			if (GetHash(Index) == InInHash)
+			{
+				return Index;
+			}
+		}
+
+		return INDEX_NONE;
+	}
 };
