Advanced Renamer
- Fixed bug where you could not execute renames

diff --git a/Engine/Plugins/Experimental/AdvancedRenamer/Source/AdvancedRenamer/Private/AdvancedRenamer.cpp b/Engine/Plugins/Experimental/AdvancedRenamer/Source/AdvancedRenamer/Private/AdvancedRenamer.cpp
--- a/Engine/Plugins/Experimental/AdvancedRenamer/Source/AdvancedRenamer/Private/AdvancedRenamer.cpp
+++ b/Engine/Plugins/Experimental/AdvancedRenamer/Source/AdvancedRenamer/Private/AdvancedRenamer.cpp
@@ -106,7 +106,7 @@ bool FAdvancedRenamer::UpdatePreviews()
 		// Force recreation
 		Previews[Index]->NewName = ApplyRename(Previews[Index]->OriginalName, Index);
 
-		if (!Previews[Index]->NewName.IsEmpty())
+		if (Previews[Index]->NewName.IsEmpty())
 		{
 			continue;
 		}
@@ -151,7 +151,7 @@ bool FAdvancedRenamer::Execute()
 			continue;
 		}
 
-		if (!Previews[Index]->NewName.IsEmpty())
+		if (Previews[Index]->NewName.IsEmpty())
 		{
 			continue;
 		}
