[AdvancedRenamer] Re-Style
Updating AdvancedRenamer style plugin to be more consistent with the rest of the engine.
Updating AdvancedRenamer backend to be more flexible and extensible in the future.
Renaming AdvancedRenamer to BatchRenamer in friendly uplugin name and UI.

diff --git a/Engine/Plugins/Experimental/AdvancedRenamer/AdvancedRenamer.uplugin b/Engine/Plugins/Experimental/AdvancedRenamer/AdvancedRenamer.uplugin
--- a/Engine/Plugins/Experimental/AdvancedRenamer/AdvancedRenamer.uplugin
+++ b/Engine/Plugins/Experimental/AdvancedRenamer/AdvancedRenamer.uplugin
@@ -2,8 +2,8 @@
 	"FileVersion": 3,
 	"Version": 1,
 	"VersionName": "1.0",
-	"FriendlyName": "Advanced Renamer",
-	"Description": "Advanced Renaming tool for actors and assets",
+	"FriendlyName": "Batch Renamer",
+	"Description": "Batch Renaming tool for actors and assets",
 	"Category": "Experimental",
 	"CreatedBy": "Epic Games, Inc.",
 	"CreatedByURL": "https://epicgames.com",
@@ -13,6 +13,7 @@
 	"CanContainContent": false,
 	"IsBetaVersion": false,
 	"IsExperimentalVersion": true,
+	"EnabledByDefault": true,
 	"Installed": false,
 	"Modules": [
 		{
diff --git a/Engine/Plugins/Experimental/AdvancedRenamer/Source/AdvancedRenamer/Private/AdvancedRenamer.cpp b/Engine/Plugins/Experimental/AdvancedRenamer/Source/AdvancedRenamer/Private/AdvancedRenamer.cpp
--- a/Engine/Plugins/Experimental/AdvancedRenamer/Source/AdvancedRenamer/Private/AdvancedRenamer.cpp
+++ b/Engine/Plugins/Experimental/AdvancedRenamer/Source/AdvancedRenamer/Private/AdvancedRenamer.cpp
@@ -3,16 +3,10 @@
 #include "AdvancedRenamer.h"
 
 #include "AdvancedRenamerModule.h"
-#include "Internationalization/Regex.h"
 
-FAdvancedRenamer::FAdvancedRenamer(const TSharedRef<IAdvancedRenamerProvider>& InProvider, FAdvancedRenamerOptions* InInitialOptions)
+FAdvancedRenamer::FAdvancedRenamer(const TSharedRef<IAdvancedRenamerProvider>& InProvider)
 	: Provider(InProvider)
 {
-	if (InInitialOptions)
-	{
-		SetOption(*InInitialOptions);
-	}
-
 	int32 Count = Num();
 	check(Count > 0);
 
@@ -38,22 +32,6 @@ const TSharedRef<IAdvancedRenamerProvider>& FAdvancedRenamer::GetProvider() cons
 	return Provider;
 }
 
-const FAdvancedRenamerOptions& FAdvancedRenamer::GetOptions() const
-{
-	return Options;
-}
-
-FAdvancedRenamerOptions& FAdvancedRenamer::GetOptions()
-{
-	return Options;
-}
-
-void FAdvancedRenamer::SetOption(const FAdvancedRenamerOptions& InOptions)
-{
-	Options = InOptions;
-	MarkDirty();
-}
-
 const TArray<TSharedPtr<FAdvancedRenamerPreview>>& FAdvancedRenamer::GetPreviews()
 {
 	return Previews;
@@ -69,6 +47,11 @@ TSharedPtr<FAdvancedRenamerPreview> FAdvancedRenamer::GetPreview(int32 InIndex)
 	return nullptr;
 }
 
+void FAdvancedRenamer::AddSection(FAdvancedRenamerExecuteSection InSection)
+{
+	Sections.Add(InSection);
+}
+
 bool FAdvancedRenamer::HasRenames() const
 {
 	return bHasRenames;
@@ -94,6 +77,8 @@ bool FAdvancedRenamer::UpdatePreviews()
 	bHasRenames = false;
 	const int32 Count = Previews.Num();
 
+	BeforeOperationsStartExecute();
+
 	for (int32 Index = 0; Index < Count; ++Index)
 	{
 		if (!Previews[Index].IsValid() || !IsValidIndex(Index))
@@ -104,14 +89,14 @@ bool FAdvancedRenamer::UpdatePreviews()
 		}
 
 		// Force recreation
-		Previews[Index]->NewName = ApplyRename(Previews[Index]->OriginalName, Index);
+		Previews[Index]->NewName = ApplyRename(Previews[Index]->OriginalName);
 
 		if (Previews[Index]->NewName.IsEmpty())
 		{
 			continue;
 		}
 
-		if (GetOriginalName(Index) == Previews[Index]->NewName)
+		if (GetOriginalName(Index).Equals(Previews[Index]->NewName))
 		{
 			continue;
 		}
@@ -119,6 +104,8 @@ bool FAdvancedRenamer::UpdatePreviews()
 		bHasRenames = true;
 	}
 
+	AfterOperationsEndExecute();
+
 	MarkClean();
 
 	return bHasRenames;
@@ -167,283 +154,30 @@ bool FAdvancedRenamer::Execute()
 	return bAllSuccess;
 }
 
-FString FAdvancedRenamer::ApplyRename(const FString& InOriginalName, int32 InIndex) const
+FString FAdvancedRenamer::ApplyRename(const FString& InOriginalName)
 {
 	FString NewName = InOriginalName;
-
-	if (Options.HasBaseName())
-	{
-		NewName = ApplyBaseName(NewName);
-	}
-
-	if (Options.HasPrefix())
-	{
-		NewName = ApplyPrefix(NewName);
-	}
-
-	if (Options.HasSuffix())
-	{
-		NewName = ApplySuffix(NewName, InIndex);
-	}
-
-	if (Options.HasSearchAndReplace())
+	for (FAdvancedRenamerExecuteSection& Section : Sections)
 	{
-		switch (Options.SearchAndReplaceType)
-		{
-			default:
-				break;
-
-			case EAdvancedRenamerSeachAndReplaceType::PlainText:
-				NewName = ApplySearchPlainText(NewName);
-				break;
-
-			case EAdvancedRenamerSeachAndReplaceType::RegularExpression:
-				NewName = ApplySearchReplaceRegex(NewName);
-				break;
-		}
+		Section.OnOperationExecuted().ExecuteIfBound(NewName);
 	}
-
 	return NewName;
 }
 
-FString FAdvancedRenamer::ApplyBaseName(const FString& InOriginalName) const
+void FAdvancedRenamer::BeforeOperationsStartExecute()
 {
-	if (Options.BaseName.IsEmpty())
+	for (FAdvancedRenamerExecuteSection& Section : Sections)
 	{
-		return InOriginalName;
+		Section.OnBeforeOperationExecutionStart().ExecuteIfBound();
 	}
-
-	return Options.BaseName;
 }
 
-FString FAdvancedRenamer::ApplyPrefix(const FString& InOriginalName) const
+void FAdvancedRenamer::AfterOperationsEndExecute()
 {
-	FString Output = InOriginalName;
-
-	if (!Options.RemovePrefixSeparator.IsEmpty())
+	for (FAdvancedRenamerExecuteSection& Section : Sections)
 	{
-		const int32 PrefixStart = Output.Find(Options.RemovePrefixSeparator, ESearchCase::IgnoreCase);
-
-		if (PrefixStart >= 0)
-		{
-			Output = Output.Mid(PrefixStart + Options.RemovePrefixSeparator.Len());
-		}
+		Section.OnAfterOperationExecutionEnded().ExecuteIfBound();
 	}
-	else if (Options.RemovePrefixCharacterCount > 0)
-	{
-		Output = Output.RightChop(Options.RemovePrefixCharacterCount);
-	}
-
-	if (!Options.AddPrefix.IsEmpty())
-	{
-		Output = Options.AddPrefix + Output;
-	}
-
-	return Output;
-}
-
-FString FAdvancedRenamer::ApplySuffix(const FString& InOriginalName, int32 InIndex) const
-{
-	static constexpr TCHAR FirstDigit = '0';
-	static constexpr TCHAR LastDigit = '9';
-
-	FString Output = InOriginalName;
-
-	if (!Options.RemoveSuffixSeparator.IsEmpty())
-	{
-		const int32 SuffixStart = Output.Find(Options.RemoveSuffixSeparator, ESearchCase::IgnoreCase, ESearchDir::FromEnd);
-
-		if (SuffixStart >= 0)
-		{
-			Output = Output.Mid(0, SuffixStart);
-		}
-	}
-	else if (Options.RemoveSuffixCharacterCount > 0)
-	{
-		Output = Output.LeftChop(Options.RemoveSuffixCharacterCount);
-	}
-
-	if (Options.bRemoveSuffixNumber || Options.bAddSuffixNumber)
-	{
-		int32 LastDigitIndex = Output.Len() - 1;
-
-		while (LastDigitIndex >= 0)
-		{
-			if (Output[LastDigitIndex] < FirstDigit || Output[LastDigitIndex] > LastDigit)
-			{
-				break;
-			}
-
-			--LastDigitIndex;
-		}
-
-		Output = Output.Mid(0, LastDigitIndex + 1);
-	}
-
-	if (!Options.AddSuffix.IsEmpty())
-	{
-		Output = Output + Options.AddSuffix;
-	}
-
-	if (Options.bAddSuffixNumber)
-	{
-		const int32 Current = Options.AddSuffixNumberStart + (InIndex * Options.AddSuffixNumberStep);
-		Output += FString::FromInt(Current);
-	}
-
-	return Output;
-}
-
-FString FAdvancedRenamer::ApplySearchPlainText(const FString& InOriginalName) const
-{
-	if (Options.SearchAndReplaceFromText.IsEmpty())
-	{
-		return InOriginalName;
-	}
-
-	return InOriginalName.Replace(
-		*Options.SearchAndReplaceFromText, 
-		*Options.SearchAndReplaceToText, 
-		Options.SearchAndReplaceCase
-	);
-}
-
-FString FAdvancedRenamer::ApplySearchReplaceRegex(const FString& InOriginalName) const
-{
-	if (Options.SearchAndReplaceFromText.IsEmpty())
-	{
-		return InOriginalName;
-	}
-
-	const FRegexPattern RegexPattern = FRegexPattern(
-		Options.SearchAndReplaceFromText,
-		Options.SearchAndReplaceCase == ESearchCase::IgnoreCase ? ERegexPatternFlags::CaseInsensitive : ERegexPatternFlags::None
-	);
-
-	return RegexReplace(InOriginalName, RegexPattern, Options.SearchAndReplaceToText);
-}
-
-FString FAdvancedRenamer::RegexReplace(const FString& InOriginalString, const FRegexPattern& InPattern, const FString& InReplaceString) const
-{
-	static const FString EscapeString = TEXT("\\");
-	static constexpr TCHAR EscapeChar = '\\';
-	static const FString GroupString = TEXT("$");
-	static constexpr TCHAR GroupChar = '$';
-	static constexpr TCHAR FirstDigit = '0';
-	static constexpr TCHAR LastDigit = '9';
-	static constexpr TCHAR NullChar = 0;
-
-	FRegexMatcher Matcher(InPattern, InOriginalString);
-
-	FString Output = "";
-	int32 StartCharIndex = 0;
-	bool bReadingGroupName = false;
-	int32 GroupIndex = INDEX_NONE;
-
-	while (Matcher.FindNext())
-	{
-		// Add on part of string after start/previous match
-		if (StartCharIndex != Matcher.GetMatchBeginning())
-		{
-			Output += InOriginalString.Mid(StartCharIndex, Matcher.GetMatchBeginning() - StartCharIndex);
-		}
-
-		bool bEscaped = false;
-
-		for (int32 CharIndex = 0; CharIndex <= InReplaceString.Len(); ++CharIndex)
-		{
-			const TCHAR& Char = CharIndex < InReplaceString.Len() ? InReplaceString[CharIndex] : NullChar;
-
-			if (bReadingGroupName)
-			{
-				// Build group index
-				if (Char >= FirstDigit && Char <= LastDigit)
-				{
-					if (GroupIndex == INDEX_NONE)
-					{
-						GroupIndex = 0;
-					}
-					else if (GroupIndex > 0)
-					{
-						GroupIndex *= 10;
-					}
-
-					const int32 NextDigit = static_cast<int32>(Char - FirstDigit);
-					GroupIndex += NextDigit;
-					continue;
-				}
-				// We've read a group index, add it to the output string
-				else if (GroupIndex > 0)
-				{
-					if (Matcher.GetCaptureGroupBeginning(GroupIndex) == INDEX_NONE)
-					{
-						UE_LOG(LogARP, Error, TEXT("Regex: Capture group does not exist %d."), GroupIndex);
-					}
-
-					Output += Matcher.GetCaptureGroup(GroupIndex);
-				}
-				// $0 matches the entire matched string
-				else if (GroupIndex == 0)
-				{
-					Output += InOriginalString.Mid(Matcher.GetMatchBeginning(), Matcher.GetMatchEnding() - Matcher.GetMatchBeginning());
-				}
-				// An unescaped $
-				else
-				{
-					UE_LOG(LogARP, Error, TEXT("Regex: Unescaped %s."), *GroupString);
-
-					Output += GroupString;
-				}
-
-				bReadingGroupName = false;
-				// Continue regular parsing of this character.
-			}
-
-			// Check for special chars
-			if (!bEscaped)
-			{
-				if (Char == EscapeChar)
-				{
-					bEscaped = true;
-					continue;
-				}
-
-				if (Char == GroupChar)
-				{
-					bReadingGroupName = true;
-					GroupIndex = INDEX_NONE;
-					continue;
-				}
-			}
-			else
-			{
-				// If the last char is a \ assume that it's not an escape char
-				if (Char == NullChar)
-				{
-					UE_LOG(LogARP, Error, TEXT("Regex: Unescaped %s."), *EscapeString);
-
-					Output += EscapeChar;
-				}
-			}
-
-			if (Char != NullChar)
-			{
-				Output += Char;
-			}
-
-			bEscaped = false;
-		}
-
-		StartCharIndex = Matcher.GetMatchEnding();
-	}
-
-	// Add on the end of the string after the last match
-	if (StartCharIndex < InOriginalString.Len())
-	{
-		Output += InOriginalString.Mid(StartCharIndex);
-	}
-
-	return Output;
 }
 
 int32 FAdvancedRenamer::Num() const
diff --git a/Engine/Plugins/Experimental/AdvancedRenamer/Source/AdvancedRenamer/Private/AdvancedRenamer.h b/Engine/Plugins/Experimental/AdvancedRenamer/Source/AdvancedRenamer/Private/AdvancedRenamer.h
--- a/Engine/Plugins/Experimental/AdvancedRenamer/Source/AdvancedRenamer/Private/AdvancedRenamer.h
+++ b/Engine/Plugins/Experimental/AdvancedRenamer/Source/AdvancedRenamer/Private/AdvancedRenamer.h
@@ -6,27 +6,23 @@
 #include "Providers/IAdvancedRenamerProvider.h"
 #include "Templates/SharedPointer.h"
 
-class FRegexPattern;
-struct FAdvancedRenamerOptions;
 struct FAdvancedRenamerPreview;
 
 class FAdvancedRenamer : public IAdvancedRenamer
 {
 public:
-	FAdvancedRenamer(const TSharedRef<IAdvancedRenamerProvider>& InProvider, FAdvancedRenamerOptions* InInitialOptions = nullptr);
+	FAdvancedRenamer(const TSharedRef<IAdvancedRenamerProvider>& InProvider);
 
 	//~ Begin IAdvancedRenamer
 	virtual const TSharedRef<IAdvancedRenamerProvider>& GetProvider() const override;
-	virtual const FAdvancedRenamerOptions& GetOptions() const override;
-	virtual FAdvancedRenamerOptions& GetOptions() override;
-	virtual void SetOption(const FAdvancedRenamerOptions& InOptions) override;
 	virtual const TArray<TSharedPtr<FAdvancedRenamerPreview>>& GetPreviews() override;
 	virtual TSharedPtr<FAdvancedRenamerPreview> GetPreview(int32 InIndex) const override;
+	virtual void AddSection(FAdvancedRenamerExecuteSection InSection) override;
 	virtual bool HasRenames() const override;
 	virtual bool IsDirty() const override;
 	virtual void MarkDirty() override;
 	virtual void MarkClean() override;
-	virtual FString ApplyRename(const FString& InName, int32 InIndex) const override;
+	virtual FString ApplyRename(const FString& InName) override;
 	virtual bool UpdatePreviews() override;
 	virtual bool Execute() override;
 	//~ End IAdvancedRenamer
@@ -41,18 +37,26 @@ public:
 	virtual bool ExecuteRename(int32 InIndex, const FString& InNewName) override;
 	//~ End IAdvancedRenamerProvider
 
-protected:
+private:
+	/** Called before the whole Rename logic start */
+	void BeforeOperationsStartExecute();
+	
+	/** Called after the whole Rename logic end */
+	void AfterOperationsEndExecute();
+
+private:
+	/** Provider for this Renamer */
 	TSharedRef<IAdvancedRenamerProvider> Provider;
+
+	/** Previews Name list */
 	TArray<TSharedPtr<FAdvancedRenamerPreview>> Previews;
-	FAdvancedRenamerOptions Options;
+	
+	/** Renamer sections list */
+	TArray<FAdvancedRenamerExecuteSection> Sections;
+
+	/** Whether or not at least 1 preview has a Rename */
 	bool bHasRenames;
+	
+	/** if true the Rename logic will be executed */
 	bool bDirty;
-
-	FString ApplyBaseName(const FString& InOriginalName) const;
-	FString ApplyPrefix(const FString& InOriginalName) const;
-	FString ApplySuffix(const FString& InOriginalName, int32 InIndex) const;
-	FString ApplySearchPlainText(const FString& InOriginalName) const;
-	FString ApplySearchReplaceRegex(const FString& InOriginalName) const;
-
-	FString RegexReplace(const FString& InOriginalString, const FRegexPattern& InPattern, const FString& InReplaceString) const;
 };
diff --git a/Engine/Plugins/Experimental/AdvancedRenamer/Source/AdvancedRenamer/Private/AdvancedRenamerCommands.cpp b/Engine/Plugins/Experimental/AdvancedRenamer/Source/AdvancedRenamer/Private/AdvancedRenamerCommands.cpp
--- a/Engine/Plugins/Experimental/AdvancedRenamer/Source/AdvancedRenamer/Private/AdvancedRenamerCommands.cpp
+++ b/Engine/Plugins/Experimental/AdvancedRenamer/Source/AdvancedRenamer/Private/AdvancedRenamerCommands.cpp
@@ -22,13 +22,13 @@ void FAdvancedRenamerCommands::RegisterCommands()
 {
 	UI_COMMAND(RenameSelectedActors
 		, "Rename Selected Actors"
-		, "Opens the Advanced Renamer Panel to rename all selected actors."
+		, "Opens the Batch Renamer Panel to rename all selected actors."
 		, EUserInterfaceActionType::Button
 		, FInputChord())
 
 	UI_COMMAND(RenameSharedClassActors
 		, "Rename Actors of Selected Actor Classes"
-		, "Opens the Advanced Renamer Panel to rename all actors sharing a class with any selected actor."
+		, "Opens the Batch Renamer Panel to rename all actors sharing a class with any selected actor."
 		, EUserInterfaceActionType::Button
 		, FInputChord())
 }
diff --git a/Engine/Plugins/Experimental/AdvancedRenamer/Source/AdvancedRenamer/Private/AdvancedRenamerExecuteSection.h b/Engine/Plugins/Experimental/AdvancedRenamer/Source/AdvancedRenamer/Private/AdvancedRenamerExecuteSection.h
new file mode 100644
--- /dev/null
+++ b/Engine/Plugins/Experimental/AdvancedRenamer/Source/AdvancedRenamer/Private/AdvancedRenamerExecuteSection.h
@@ -0,0 +1,36 @@
+﻿// Copyright Epic Games, Inc. All Rights Reserved.
+
+#pragma once
+
+#include "Delegates/Delegate.h"
+#include "UObject/NameTypes.h"
+
+DECLARE_DELEGATE(FOnBeforeOperationExecutionStart)
+DECLARE_DELEGATE_OneParam(FOnOperationExecuted, FString& /** OutString */)
+DECLARE_DELEGATE(FOnAfterOperationExecutionEnded)
+
+struct FAdvancedRenamerExecuteSection
+{
+	/** Get the callback executed before the whole Rename happens */
+	FOnOperationExecuted& OnOperationExecuted() { return OnOperationExecutedDelegate; }
+
+	/** Get the callback executed during the Rename */
+	FOnBeforeOperationExecutionStart& OnBeforeOperationExecutionStart() { return OnBeforeOperationExecutionStartDelegate; }
+
+	/** Get the callback executed after the whole Rename happens */
+	FOnAfterOperationExecutionEnded& OnAfterOperationExecutionEnded() { return OnAfterOperationExecutionEndedDelegate; }
+
+public:
+	/** Name of this section */
+	FName SectionName;
+
+private:
+	/** Callback called before the whole Rename execution is applied */
+	FOnBeforeOperationExecutionStart OnBeforeOperationExecutionStartDelegate;
+
+	/** Callback called during the whole Rename execution */
+	FOnOperationExecuted OnOperationExecutedDelegate;
+
+	/** Callback called after the whole Rename execution is applied */
+	FOnAfterOperationExecutionEnded OnAfterOperationExecutionEndedDelegate;
+};
diff --git a/Engine/Plugins/Experimental/AdvancedRenamer/Source/AdvancedRenamer/Private/AdvancedRenamerModule.cpp b/Engine/Plugins/Experimental/AdvancedRenamer/Source/AdvancedRenamer/Private/AdvancedRenamerModule.cpp
--- a/Engine/Plugins/Experimental/AdvancedRenamer/Source/AdvancedRenamer/Private/AdvancedRenamerModule.cpp
+++ b/Engine/Plugins/Experimental/AdvancedRenamer/Source/AdvancedRenamer/Private/AdvancedRenamerModule.cpp
@@ -4,6 +4,11 @@
 
 #include "AdvancedRenamer.h"
 #include "AdvancedRenamerCommands.h"
+#include "AdvancedRenamerSections/AdvancedRenamerAddPrefixSuffixSection.h"
+#include "AdvancedRenamerSections/AdvancedRenamerChangeCaseSection.h"
+#include "AdvancedRenamerSections/AdvancedRenamerRemovePrefixSection.h"
+#include "AdvancedRenamerSections/AdvancedRenamerRemoveSuffixSection.h"
+#include "AdvancedRenamerSections/AdvancedRenamerSearchAndReplaceSection.h"
 #include "AdvancedRenamerStyle.h"
 #include "EngineUtils.h"
 #include "Framework/Application/SlateApplication.h"
@@ -23,23 +28,59 @@ namespace UE::AdvancedRenamer::Private
 	TSharedRef<SWindow> CreateAdvancedRenamerWindow()
 	{
 		return SNew(SWindow)
-			.Title(LOCTEXT("AdvancedRenameWindow", "Rename Actors"))
-			.ClientSize(FVector2D(600, 538))
+			.Title(LOCTEXT("AdvancedRenameWindow", "Batch Renamer"))
+			.ClientSize(FVector2D(731, 547))
 			.SizingRule(ESizingRule::UserSized)
 			.SupportsMaximize(false)
 			.SupportsMinimize(false)
-			.MinWidth(600)
-			.MinHeight(538)
-			.MaxHeight(538);
+			.MinWidth(737)
+			.MinHeight(586.5f);
+	}
+
+	static TAutoConsoleVariable<bool> CVarBatchRenamerEnableIntegration(
+		TEXT("BatchRenamer.EnableIntegration"),
+		false,
+		TEXT("If enabled, the BatchRenamer will hook into the ContentBrowser and Outliner.")
+	);
+
+	static bool bPreviousCVarValue = false;
+
+	void OnBatchRenamerEnabledChanged(IConsoleVariable* InVar)
+	{
+		if (InVar)
+		{
+			const bool bCurrentValue = InVar->GetBool();
+			if (bCurrentValue != bPreviousCVarValue)
+			{
+				if (InVar->GetBool())
+				{
+					FAdvancedRenamerContentBrowserIntegration::Initialize();
+					FAdvancedRenamerLevelEditorIntegration::InitializeMenu();
+				}
+				else
+				{
+					FAdvancedRenamerContentBrowserIntegration::Shutdown();
+					FAdvancedRenamerLevelEditorIntegration::ShutdownMenu();
+				}
+				bPreviousCVarValue = bCurrentValue;
+			}
+		}
 	}
 }
 
 void FAdvancedRenamerModule::StartupModule()
 {
 	FAdvancedRenamerStyle::Initialize();
-	FAdvancedRenamerCommands::Register();
-	FAdvancedRenamerContentBrowserIntegration::Initialize();
 	FAdvancedRenamerLevelEditorIntegration::Initialize();
+	FAdvancedRenamerCommands::Register();
+
+	{
+		// Enable plugins integration section from CVar
+		using namespace UE::AdvancedRenamer::Private;
+		EnableRenamerHandle = CVarBatchRenamerEnableIntegration.AsVariable()->OnChangedDelegate().AddStatic(&OnBatchRenamerEnabledChanged);
+	}
+
+	RegisterDefaultSections();
 }
 
 void FAdvancedRenamerModule::ShutdownModule()
@@ -48,6 +89,12 @@ void FAdvancedRenamerModule::ShutdownModule()
 	FAdvancedRenamerStyle::Shutdown();
 	FAdvancedRenamerContentBrowserIntegration::Shutdown();
 	FAdvancedRenamerLevelEditorIntegration::Shutdown();
+	FAdvancedRenamerLevelEditorIntegration::ShutdownMenu();
+	if (EnableRenamerHandle.IsValid())
+	{
+		using namespace UE::AdvancedRenamer::Private;
+		CVarBatchRenamerEnableIntegration.AsVariable()->OnChangedDelegate().Remove(EnableRenamerHandle);
+	}
 }
 
 TSharedRef<IAdvancedRenamer> FAdvancedRenamerModule::CreateAdvancedRenamer(const TSharedRef<IAdvancedRenamerProvider>& InRenameProvider)
@@ -205,6 +252,15 @@ TArray<AActor*> FAdvancedRenamerModule::GetActorsSharingClassesInWorld(const TAr
 	return AllActors;
 }
 
+void FAdvancedRenamerModule::RegisterDefaultSections()
+{
+	Sections.Add(IAdvancedRenamerSection::MakeInstance<FAdvancedRenamerSearchAndReplaceSection>());
+	Sections.Add(IAdvancedRenamerSection::MakeInstance<FAdvancedRenamerRemovePrefixSection>());
+	Sections.Add(IAdvancedRenamerSection::MakeInstance<FAdvancedRenamerRemoveSuffixSection>());
+	Sections.Add(IAdvancedRenamerSection::MakeInstance<FAdvancedRenamerAddPrefixSuffixSection>());
+	Sections.Add(IAdvancedRenamerSection::MakeInstance<FAdvancedRenamerChangeCaseSection>());
+}
+
 IMPLEMENT_MODULE(FAdvancedRenamerModule, AdvancedRenamer)
 
 #undef LOCTEXT_NAMESPACE
diff --git a/Engine/Plugins/Experimental/AdvancedRenamer/Source/AdvancedRenamer/Private/AdvancedRenamerModule.h b/Engine/Plugins/Experimental/AdvancedRenamer/Source/AdvancedRenamer/Private/AdvancedRenamerModule.h
--- a/Engine/Plugins/Experimental/AdvancedRenamer/Source/AdvancedRenamer/Private/AdvancedRenamerModule.h
+++ b/Engine/Plugins/Experimental/AdvancedRenamer/Source/AdvancedRenamer/Private/AdvancedRenamerModule.h
@@ -3,6 +3,8 @@
 #pragma once
 
 #include "IAdvancedRenamerModule.h"
+#include "AdvancedRenamerSections/IAdvancedRenamerSection.h"
+#include "Delegates/IDelegateInstance.h"
 #include "Logging/LogMacros.h"
 
 DECLARE_LOG_CATEGORY_EXTERN(LogARP, Log, All);
@@ -24,5 +26,13 @@ class FAdvancedRenamerModule : public IAdvancedRenamerModule
 	virtual void OpenAdvancedRenamer(const TSharedRef<IAdvancedRenamer>& InRenamer, const TSharedPtr<IToolkitHost>& InToolkitHost) override;
 	virtual void OpenAdvancedRenamer(const TSharedRef<IAdvancedRenamer>& InRenamer, const TSharedPtr<SWidget>& InParentWidget) override;
 	virtual TArray<AActor*> GetActorsSharingClassesInWorld(const TArray<AActor*>& InActors) override;
+	virtual TArray<TSharedPtr<IAdvancedRenamerSection>> GetRegisteredSections() const override { return Sections; }
 	//~ End IAdvancedRenamerModule
+
+private:
+	void RegisterDefaultSections();
+
+private:
+	TArray<TSharedPtr<IAdvancedRenamerSection>> Sections;
+	FDelegateHandle EnableRenamerHandle;
 };
diff --git a/Engine/Plugins/Experimental/AdvancedRenamer/Source/AdvancedRenamer/Private/AdvancedRenamerSections/AdvancedRenamerAddPrefixSuffixSection.cpp b/Engine/Plugins/Experimental/AdvancedRenamer/Source/AdvancedRenamer/Private/AdvancedRenamerSections/AdvancedRenamerAddPrefixSuffixSection.cpp
new file mode 100644
--- /dev/null
+++ b/Engine/Plugins/Experimental/AdvancedRenamer/Source/AdvancedRenamer/Private/AdvancedRenamerSections/AdvancedRenamerAddPrefixSuffixSection.cpp
@@ -0,0 +1,340 @@
+// Copyright Epic Games, Inc. All Rights Reserved.
+
+#include "AdvancedRenamerAddPrefixSuffixSection.h"
+
+#include "AdvancedRenamerStyle.h"
+#include "IAdvancedRenamer.h"
+#include "Styling/StyleColors.h"
+#include "Utils/AdvancedRenamerSlateUtils.h"
+#include "Widgets/Input/SCheckBox.h"
+#include "Widgets/Input/SEditableTextBox.h"
+#include "Widgets/Input/SSpinBox.h"
+#include "Widgets/Layout/SBorder.h"
+#include "Widgets/Layout/SBox.h"
+#include "Widgets/Layout/SSpacer.h"
+
+#define LOCTEXT_NAMESPACE "AdvancedRenamerAddPrefixSuffixSection"
+
+FAdvancedRenamerAddPrefixSuffixSection::FAdvancedRenamerAddPrefixSuffixSection()
+{
+	FAdvancedRenamerAddPrefixSuffixSection::ResetToDefault();
+}
+
+void FAdvancedRenamerAddPrefixSuffixSection::Init(TSharedRef<IAdvancedRenamer> InRenamer)
+{
+	FAdvancedRenamerSectionBase::Init(InRenamer);
+	Section.SectionName = TEXT("AddPrefixSuffixNumber");
+	Section.OnBeforeOperationExecutionStart().BindSP(this, &FAdvancedRenamerAddPrefixSuffixSection::ResetCurrentSuffixNumber);
+	Section.OnOperationExecuted().BindSP(this, &FAdvancedRenamerAddPrefixSuffixSection::ApplyAddPrefixSuffixNumberOperation);
+	InRenamer->AddSection(Section);
+}
+
+TSharedRef<SWidget> FAdvancedRenamerAddPrefixSuffixSection::GetWidget()
+{
+	using namespace AdvancedRenamerSlateUtils::Default;
+
+	return SNew(SBorder)
+		.BorderImage(FAdvancedRenamerStyle::Get().GetBrush("AdvancedRenamer.Style.BackgroundBorder"))
+		.Content()
+		[
+			SNew(SVerticalBox)
+
+			//Add Prefix
+			+ SVerticalBox::Slot()
+			.Padding(SectionContentFirstEntryPadding)
+			.AutoHeight()
+			[
+				CreateAddPrefix()
+			]
+
+			// Add Suffix
+			+ SVerticalBox::Slot()
+			.Padding(SectionContentMiddleEntriesPadding)
+			.AutoHeight()
+			[
+				CreateAddSuffix()
+			]
+			
+			// Add Number
+			+ SVerticalBox::Slot()
+			.Padding(SectionContentMiddleEntriesPadding)
+			.AutoHeight()
+			[
+				CreateAddNumber()
+			]
+		];
+}
+
+void FAdvancedRenamerAddPrefixSuffixSection::ResetToDefault()
+{
+	PrefixText = FText::GetEmpty();
+	SuffixText = FText::GetEmpty();
+	bAddSuffixNumbers = false;
+	SuffixNumberStartValue = 1;
+	CurrentSuffixNumber = 1;
+	SuffixNumberStepValue = 1;
+}
+
+TSharedRef<SWidget> FAdvancedRenamerAddPrefixSuffixSection::CreateAddPrefix()
+{
+	using namespace AdvancedRenamerSlateUtils::Default;
+	
+	return SNew(SHorizontalBox)
+		+ SHorizontalBox::Slot()
+		.HAlign(HAlign_Left)
+		.VAlign(VAlign_Center)
+		.AutoWidth()
+		.Padding(NameWidgetPadding)
+		[
+			SNew(SBox)
+			.WidthOverride(70.f)
+			[
+				SNew(STextBlock)
+				.Font(RegularFont)
+				.Text(LOCTEXT("AR_AddPrefix", "Add Prefix"))
+			]
+		]
+		
+		+ SHorizontalBox::Slot()
+		.VAlign(VAlign_Center)
+		.FillWidth(1.f)
+		.Padding(ValueWidgetPadding)
+		[
+			SAssignNew(PrefixTextBox, SEditableTextBox)
+			.BackgroundColor(FStyleColors::Background.GetSpecifiedColor())
+			.Font(RegularFont)
+			.HintText(LOCTEXT("AR_PrefixHint", "New Prefix"))
+			.Text(this, &FAdvancedRenamerAddPrefixSuffixSection::GetPrefixText)
+			.OnTextChanged(this, &FAdvancedRenamerAddPrefixSuffixSection::OnPrefixChanged)
+		];
+}
+
+TSharedRef<SWidget> FAdvancedRenamerAddPrefixSuffixSection::CreateAddSuffix()
+{
+	using namespace AdvancedRenamerSlateUtils::Default;
+
+	return SNew(SHorizontalBox)
+		+ SHorizontalBox::Slot()
+		.HAlign(HAlign_Left)
+		.VAlign(VAlign_Center)
+		.AutoWidth()
+		.Padding(NameWidgetPadding)
+		[
+			SNew(SBox)
+			.WidthOverride(70.f)
+			[
+				SNew(STextBlock)
+				.Font(RegularFont)
+				.Text(LOCTEXT("BR_AddSuffix", "Add Suffix"))
+			]
+		]
+		+ SHorizontalBox::Slot()
+		.VAlign(VAlign_Center)
+		.FillWidth(1.f)
+		.Padding(ValueWidgetPadding)
+		[
+			SAssignNew(SuffixTextBox, SEditableTextBox)
+			.BackgroundColor(FStyleColors::Background.GetSpecifiedColor())
+			.Font(RegularFont)
+			.HintText(LOCTEXT("AR_SuffixHint", "New Suffix"))
+			.Text(this, &FAdvancedRenamerAddPrefixSuffixSection::GetSuffixText)
+			.OnTextChanged(this, &FAdvancedRenamerAddPrefixSuffixSection::OnSuffixChanged)
+		];
+}
+
+TSharedRef<SWidget> FAdvancedRenamerAddPrefixSuffixSection::CreateAddNumber()
+{
+	using namespace AdvancedRenamerSlateUtils::Default;
+
+	return SNew(SHorizontalBox)
+		+ SHorizontalBox::Slot()
+		.HAlign(HAlign_Left)
+		.VAlign(VAlign_Center)
+		.AutoWidth()
+		.Padding(NameWidgetPadding)
+		[
+			SAssignNew(SuffixNumberCheckBox, SCheckBox)
+			.IsChecked(this, &FAdvancedRenamerAddPrefixSuffixSection::IsSuffixNumberChecked)
+			.OnCheckStateChanged(this, &FAdvancedRenamerAddPrefixSuffixSection::OnSuffixNumberCheckBoxChanged)
+		]
+		+ SHorizontalBox::Slot()
+		.HAlign(HAlign_Left)
+		.VAlign(VAlign_Center)
+		.AutoWidth()
+		[
+			SNew(STextBlock)
+			.Font(RegularFont)
+			.Text(LOCTEXT("AR_SuffixNumber", "Add Number"))
+		]
+		+ SHorizontalBox::Slot()
+		.HAlign(HAlign_Right)
+		.VAlign(VAlign_Center)
+		[
+			SNew(SSpacer)
+		]
+		+ SHorizontalBox::Slot()
+		.HAlign(HAlign_Right)
+		.VAlign(VAlign_Center)
+		.AutoWidth()
+		[
+			SNew(STextBlock)
+			.Font(RegularFont)
+			.Text(LOCTEXT("AR_Start", "Start"))
+		]
+		+ SHorizontalBox::Slot()
+		.HAlign(HAlign_Right)
+		.VAlign(VAlign_Center)
+		.AutoWidth()
+		.Padding(ValueWidgetPadding)
+		[
+			SAssignNew(SuffixNumberStartSpinBox, SSpinBox<int32>)
+			.Style(&FAppStyle::Get(), "Menu.SpinBox")
+			.Font(RegularFont)
+			.MinValue(1)
+			.MaxValue(99)
+			.Value(this, &FAdvancedRenamerAddPrefixSuffixSection::GetSuffixNumberStart)
+			.IsEnabled(this, &FAdvancedRenamerAddPrefixSuffixSection::IsSuffixNumberEnabled)
+			.OnValueChanged(this, &FAdvancedRenamerAddPrefixSuffixSection::OnSuffixNumberStartChanged)
+		]
+		+ SHorizontalBox::Slot()
+		.HAlign(HAlign_Right)
+		.VAlign(VAlign_Center)
+		.AutoWidth()
+		.Padding(AddNumberStepPadding)
+		[
+			SNew(STextBlock)
+			.Font(RegularFont)
+			.Text(LOCTEXT("AR_Step", "Step"))
+		]
+		+ SHorizontalBox::Slot()
+		.HAlign(HAlign_Right)
+		.VAlign(VAlign_Center)
+		.AutoWidth()
+		.Padding(ValueWidgetPadding)
+		[
+			SAssignNew(SuffixNumberStepSpinBox, SSpinBox<int32>)
+			.Style(&FAppStyle::Get(), "Menu.SpinBox")
+			.Font(RegularFont)
+			.MinValue(1)
+			.MaxValue(99)
+			.Value(this, &FAdvancedRenamerAddPrefixSuffixSection::GetSuffixNumberStep)
+			.IsEnabled(this, &FAdvancedRenamerAddPrefixSuffixSection::IsSuffixNumberEnabled)
+			.OnValueChanged(this, &FAdvancedRenamerAddPrefixSuffixSection::OnSuffixNumberStepChanged)
+		];
+}
+
+FText FAdvancedRenamerAddPrefixSuffixSection::GetPrefixText() const
+{
+	return PrefixText;
+}
+
+FText FAdvancedRenamerAddPrefixSuffixSection::GetSuffixText() const
+{
+	return SuffixText;
+}
+
+int32 FAdvancedRenamerAddPrefixSuffixSection::GetSuffixNumberStart() const
+{
+	return SuffixNumberStartValue;
+}
+
+int32 FAdvancedRenamerAddPrefixSuffixSection::GetSuffixNumberStep() const
+{
+	return SuffixNumberStepValue;
+}
+
+ECheckBoxState FAdvancedRenamerAddPrefixSuffixSection::IsSuffixNumberChecked() const
+{
+	return bAddSuffixNumbers ? ECheckBoxState::Checked : ECheckBoxState::Unchecked;
+}
+
+bool FAdvancedRenamerAddPrefixSuffixSection::IsSuffixNumberEnabled() const
+{
+	return bAddSuffixNumbers;
+}
+
+void FAdvancedRenamerAddPrefixSuffixSection::OnPrefixChanged(const FText& InNewText)
+{
+	PrefixText = InNewText;
+	MarkRenamerDirty();
+}
+
+void FAdvancedRenamerAddPrefixSuffixSection::OnSuffixChanged(const FText& InNewText)
+{
+	SuffixText = InNewText;
+	MarkRenamerDirty();
+}
+
+void FAdvancedRenamerAddPrefixSuffixSection::OnSuffixNumberCheckBoxChanged(ECheckBoxState InNewState)
+{
+	bAddSuffixNumbers = InNewState == ECheckBoxState::Checked;
+	MarkRenamerDirty();
+}
+
+void FAdvancedRenamerAddPrefixSuffixSection::OnSuffixNumberStartChanged(int32 InNewValue)
+{
+	SuffixNumberStartValue = InNewValue;
+	MarkRenamerDirty();
+}
+
+void FAdvancedRenamerAddPrefixSuffixSection::OnSuffixNumberStepChanged(int32 InNewValue)
+{
+	SuffixNumberStepValue = InNewValue;
+	MarkRenamerDirty();
+}
+
+bool FAdvancedRenamerAddPrefixSuffixSection::CanApplyAddPrefixOperation()
+{
+	return !PrefixText.IsEmpty();
+}
+
+bool FAdvancedRenamerAddPrefixSuffixSection::CanApplyAddSuffixOperation()
+{
+	return !SuffixText.IsEmpty();
+}
+
+bool FAdvancedRenamerAddPrefixSuffixSection::CanApplyAddSuffixNumberOperation()
+{
+	return IsSuffixNumberEnabled();
+}
+
+void FAdvancedRenamerAddPrefixSuffixSection::ApplyAddPrefixOperation(FString& OutOriginalName)
+{
+	OutOriginalName = PrefixText.ToString() + OutOriginalName;
+}
+
+void FAdvancedRenamerAddPrefixSuffixSection::ApplyAddSuffixOperation(FString& OutOriginalName)
+{
+	OutOriginalName = OutOriginalName + SuffixText.ToString();
+}
+
+void FAdvancedRenamerAddPrefixSuffixSection::ApplyAddSuffixNumberOperation(FString& OutOriginalName)
+{
+	OutOriginalName = OutOriginalName + FString::FromInt(CurrentSuffixNumber);
+	CurrentSuffixNumber += SuffixNumberStepValue;
+}
+
+void FAdvancedRenamerAddPrefixSuffixSection::ResetCurrentSuffixNumber()
+{
+	CurrentSuffixNumber = SuffixNumberStartValue;
+}
+
+void FAdvancedRenamerAddPrefixSuffixSection::ApplyAddPrefixSuffixNumberOperation(FString& OutOriginalName)
+{
+	if (CanApplyAddPrefixOperation())
+	{
+		ApplyAddPrefixOperation(OutOriginalName);
+	}
+
+	if (CanApplyAddSuffixOperation())
+	{
+		ApplyAddSuffixOperation(OutOriginalName);
+	}
+
+	if (CanApplyAddSuffixNumberOperation())
+	{
+		ApplyAddSuffixNumberOperation(OutOriginalName);
+	}
+}
+
+#undef LOCTEXT_NAMESPACE
diff --git a/Engine/Plugins/Experimental/AdvancedRenamer/Source/AdvancedRenamer/Private/AdvancedRenamerSections/AdvancedRenamerAddPrefixSuffixSection.h b/Engine/Plugins/Experimental/AdvancedRenamer/Source/AdvancedRenamer/Private/AdvancedRenamerSections/AdvancedRenamerAddPrefixSuffixSection.h
new file mode 100644
--- /dev/null
+++ b/Engine/Plugins/Experimental/AdvancedRenamer/Source/AdvancedRenamer/Private/AdvancedRenamerSections/AdvancedRenamerAddPrefixSuffixSection.h
@@ -0,0 +1,85 @@
+// Copyright Epic Games, Inc. All Rights Reserved.
+
+#pragma once
+
+#include "AdvancedRenamerSectionBase.h"
+#include "Styling/SlateTypes.h"
+
+class FString;
+class SCheckBox;
+class SEditableTextBox;
+template <typename InNumericType>
+class SSpinBox;
+class SWidget;
+
+class FAdvancedRenamerAddPrefixSuffixSection : public FAdvancedRenamerSectionBase
+{
+public:
+	FAdvancedRenamerAddPrefixSuffixSection();
+
+	virtual ~FAdvancedRenamerAddPrefixSuffixSection() {}
+
+	virtual void Init(TSharedRef<IAdvancedRenamer> InRenamer) override;
+
+	virtual TSharedRef<SWidget> GetWidget() override;
+
+	virtual void ResetToDefault() override;
+
+private:
+	TSharedRef<SWidget> CreateAddPrefix();
+
+	TSharedRef<SWidget> CreateAddSuffix();
+
+	TSharedRef<SWidget> CreateAddNumber();
+
+	FText GetPrefixText() const;
+
+	FText GetSuffixText() const;
+
+	int32 GetSuffixNumberStart() const;
+
+	int32 GetSuffixNumberStep() const;
+
+	ECheckBoxState IsSuffixNumberChecked() const;
+
+	bool IsSuffixNumberEnabled() const;
+
+	void OnPrefixChanged(const FText& InNewText);
+
+	void OnSuffixChanged(const FText& InNewText);
+
+	void OnSuffixNumberCheckBoxChanged(ECheckBoxState InNewState);
+
+	void OnSuffixNumberStartChanged(int32 InNewValue);
+
+	void OnSuffixNumberStepChanged(int32 InNewValue);
+
+	bool CanApplyAddPrefixOperation();
+
+	bool CanApplyAddSuffixOperation();
+
+	bool CanApplyAddSuffixNumberOperation();
+
+	void ApplyAddPrefixOperation(FString& OutOriginalName);
+
+	void ApplyAddSuffixOperation(FString& OutOriginalName);
+
+	void ApplyAddSuffixNumberOperation(FString& OutOriginalName);
+
+	void ResetCurrentSuffixNumber();
+
+	void ApplyAddPrefixSuffixNumberOperation(FString& OutOriginalName);
+
+private:
+	TSharedPtr<SEditableTextBox> PrefixTextBox;
+	TSharedPtr<SEditableTextBox> SuffixTextBox;
+	TSharedPtr<SCheckBox> SuffixNumberCheckBox;
+	TSharedPtr<SSpinBox<int32>> SuffixNumberStartSpinBox;
+	TSharedPtr<SSpinBox<int32>> SuffixNumberStepSpinBox;
+	FText PrefixText;
+	FText SuffixText;
+	bool bAddSuffixNumbers;
+	int32 SuffixNumberStartValue;
+	int32 SuffixNumberStepValue;
+	int32 CurrentSuffixNumber;
+};
diff --git a/Engine/Plugins/Experimental/AdvancedRenamer/Source/AdvancedRenamer/Private/AdvancedRenamerSections/AdvancedRenamerChangeCaseSection.cpp b/Engine/Plugins/Experimental/AdvancedRenamer/Source/AdvancedRenamer/Private/AdvancedRenamerSections/AdvancedRenamerChangeCaseSection.cpp
new file mode 100644
--- /dev/null
+++ b/Engine/Plugins/Experimental/AdvancedRenamer/Source/AdvancedRenamer/Private/AdvancedRenamerSections/AdvancedRenamerChangeCaseSection.cpp
@@ -0,0 +1,248 @@
+// Copyright Epic Games, Inc. All Rights Reserved.
+
+#include "AdvancedRenamerChangeCaseSection.h"
+
+#include "AdvancedRenamerStyle.h"
+#include "IAdvancedRenamer.h"
+#include "Utils/AdvancedRenamerSlateUtils.h"
+#include "Widgets/Input/SButton.h"
+#include "Widgets/Layout/SBorder.h"
+#include "Widgets/Layout/SBox.h"
+#include "Widgets/SBoxPanel.h"
+#include "Widgets/Text/STextBlock.h"
+
+#define LOCTEXT_NAMESPACE "AdvancedRenamerChangeCaseSection"
+
+FAdvancedRenamerChangeCaseSection::FAdvancedRenamerChangeCaseSection()
+{
+	FAdvancedRenamerChangeCaseSection::ResetToDefault();
+}
+
+void FAdvancedRenamerChangeCaseSection::Init(TSharedRef<IAdvancedRenamer> InRenamer)
+{
+	FAdvancedRenamerSectionBase::Init(InRenamer);
+	Section.SectionName = TEXT("ChangeCase");
+	Section.OnOperationExecuted().BindSP(this, &FAdvancedRenamerChangeCaseSection::ApplyChangeCaseSection);
+	Section.OnAfterOperationExecutionEnded().BindSP(this, &FAdvancedRenamerChangeCaseSection::ResetButtonClicked);
+	InRenamer->AddSection(Section);
+}
+
+TSharedRef<SWidget> FAdvancedRenamerChangeCaseSection::GetWidget()
+{
+	using namespace AdvancedRenamerSlateUtils::Default;
+
+	// Add click logic
+	return SNew(SBorder)
+		.BorderImage(FAdvancedRenamerStyle::Get().GetBrush("AdvancedRenamer.Style.BackgroundBorder"))
+		.Content()
+		[
+			SNew(SVerticalBox)
+
+			//Title
+			+ SVerticalBox::Slot()
+			.HAlign(HAlign_Left)
+			.Padding(SectionContentFirstEntryPadding)
+			.AutoHeight()
+			[
+				SNew(STextBlock)
+				.Font(TitleFont)
+				.Text(LOCTEXT("AR_ChangeCaseTitle", "Change Case"))
+			]
+
+			+ SVerticalBox::Slot()
+			.AutoHeight()
+			.Padding(SectionContentMiddleEntriesPadding)
+			[
+				SNew(SHorizontalBox)
+
+				+ SHorizontalBox::Slot()
+				.FillWidth(0.25f)
+				.VAlign(VAlign_Center)
+				.Padding(ChangeCaseFirstButtonPadding)
+				[
+					SNew(SBox)
+					.HeightOverride(25)
+					[
+						SNew(SButton)
+						.ContentPadding(FMargin(-2.f, -1.f))
+						.OnClicked(this, &FAdvancedRenamerChangeCaseSection::OnChangeCaseButtonClicked, EAdvancedRenamerChangeCaseType::SwapFirst)
+						.Content()
+						[
+							SNew(SBox)
+							.VAlign(VAlign_Center)
+							[
+								SNew(STextBlock)
+								.Justification(ETextJustify::Center)
+								.Text(LOCTEXT("AR_SwapFirst", "Swap First"))
+							]
+						]
+					]
+				]
+
+				+ SHorizontalBox::Slot()
+				.FillWidth(0.25f)
+				.VAlign(VAlign_Center)
+				.Padding(ChangeCaseMiddleButtonsPadding)
+				[
+					SNew(SBox)
+					.HeightOverride(25)
+					[
+						SNew(SButton)
+						.ContentPadding(FMargin(-2.f, -1.f))
+						.OnClicked(this, &FAdvancedRenamerChangeCaseSection::OnChangeCaseButtonClicked, EAdvancedRenamerChangeCaseType::SwapAll)
+						.Content()
+						[
+							SNew(SBox)
+							.VAlign(VAlign_Center)
+							[
+								SNew(STextBlock)
+								.Justification(ETextJustify::Center)
+								.Text(LOCTEXT("AR_SwapAll", "Swap All"))
+							]
+						]
+					]
+				]
+
+				+ SHorizontalBox::Slot()
+				.FillWidth(0.25f)
+				.VAlign(VAlign_Center)
+				.Padding(ChangeCaseMiddleButtonsPadding)
+				[
+					SNew(SBox)
+					.HeightOverride(25)
+					[
+						SNew(SButton)
+						.ContentPadding(FMargin(-2.f, -1.f))
+						.OnClicked(this, &FAdvancedRenamerChangeCaseSection::OnChangeCaseButtonClicked, EAdvancedRenamerChangeCaseType::AllLower)
+						.Content()
+						[
+							SNew(SBox)
+							.VAlign(VAlign_Center)
+							[
+								SNew(STextBlock)
+								.Justification(ETextJustify::Center)
+								.Text(LOCTEXT("AR_AllLower", "All Lower"))
+							]
+						]
+					]
+				]
+				
+				+ SHorizontalBox::Slot()
+				.FillWidth(0.25f)
+				.VAlign(VAlign_Center)
+    		    .Padding(ChangeCaseLastButtonPadding)
+    		    [
+    				SNew(SBox)
+					.HeightOverride(25)
+					[
+    		    		SNew(SButton)
+    		    		.ContentPadding(FMargin(-2.f, -1.f))
+    		    		.OnClicked(this, &FAdvancedRenamerChangeCaseSection::OnChangeCaseButtonClicked, EAdvancedRenamerChangeCaseType::AllUpper)
+    		    		.Content()
+    		    		[
+    		    			SNew(SBox)
+							.VAlign(VAlign_Center)
+							[
+    		    				SNew(STextBlock)
+    		    				.Justification(ETextJustify::Center)
+    		    				.Text(LOCTEXT("AR_AllUpper", "All Upper"))
+    		    			]
+    		    		]
+    		    	]
+    		    ]
+			]
+		];
+}
+
+void FAdvancedRenamerChangeCaseSection::ResetToDefault()
+{
+	bButtonWasClicked = false;
+	ChangeCaseType = EAdvancedRenamerChangeCaseType::SwapFirst;
+}
+
+FReply FAdvancedRenamerChangeCaseSection::OnChangeCaseButtonClicked(EAdvancedRenamerChangeCaseType InNewValue)
+{
+	ChangeCaseType = InNewValue;
+	bButtonWasClicked = true;
+	MarkRenamerDirty();
+	return FReply::Handled();
+}
+
+bool FAdvancedRenamerChangeCaseSection::CanApplyChangeCaseSection()
+{
+	return bButtonWasClicked;
+}
+
+void FAdvancedRenamerChangeCaseSection::ApplySwapFirst(FString& OutOriginalName)
+{
+	if (!OutOriginalName.IsEmpty())
+	{
+		TCHAR& First = OutOriginalName[0];
+
+		if (FChar::IsUpper(First))
+		{
+			First = FChar::ToLower(First);
+		}
+		else if (FChar::IsLower(First))
+		{
+			First = FChar::ToUpper(First);
+		}
+	}
+}
+
+void FAdvancedRenamerChangeCaseSection::ApplySwapAll(FString& OutOriginalName)
+{
+	for (int32 CurrentNameIndex = 0; CurrentNameIndex < OutOriginalName.GetCharArray().Num() - 1; CurrentNameIndex++)
+	{
+		TCHAR& CurrentChar = OutOriginalName[CurrentNameIndex];
+		if (FChar::IsUpper(CurrentChar))
+		{
+			CurrentChar = FChar::ToLower(CurrentChar);
+		}
+		else if (FChar::IsLower(CurrentChar))
+		{
+			CurrentChar = FChar::ToUpper(CurrentChar);
+		}
+	}
+}
+
+void FAdvancedRenamerChangeCaseSection::ApplyAllLower(FString& OutOriginalName)
+{
+	OutOriginalName.ToLowerInline();
+}
+
+void FAdvancedRenamerChangeCaseSection::ApplyAllUpper(FString& OutOriginalName)
+{
+	OutOriginalName.ToUpperInline();
+}
+
+void FAdvancedRenamerChangeCaseSection::ResetButtonClicked()
+{
+	bButtonWasClicked = false;
+}
+
+void FAdvancedRenamerChangeCaseSection::ApplyChangeCaseSection(FString& OutOriginalName)
+{
+	if (bButtonWasClicked)
+	{
+		switch (ChangeCaseType)
+		{
+		case EAdvancedRenamerChangeCaseType::SwapFirst:
+			ApplySwapFirst(OutOriginalName);
+			break;
+		case EAdvancedRenamerChangeCaseType::SwapAll:
+			ApplySwapAll(OutOriginalName);
+			break;
+		case EAdvancedRenamerChangeCaseType::AllLower:
+			ApplyAllLower(OutOriginalName);
+			break;
+		case EAdvancedRenamerChangeCaseType::AllUpper:
+			ApplyAllUpper(OutOriginalName);
+			break;
+		default:
+			break;
+		}
+	}
+}
+
+#undef LOCTEXT_NAMESPACE
diff --git a/Engine/Plugins/Experimental/AdvancedRenamer/Source/AdvancedRenamer/Private/AdvancedRenamerSections/AdvancedRenamerChangeCaseSection.h b/Engine/Plugins/Experimental/AdvancedRenamer/Source/AdvancedRenamer/Private/AdvancedRenamerSections/AdvancedRenamerChangeCaseSection.h
new file mode 100644
--- /dev/null
+++ b/Engine/Plugins/Experimental/AdvancedRenamer/Source/AdvancedRenamer/Private/AdvancedRenamerSections/AdvancedRenamerChangeCaseSection.h
@@ -0,0 +1,72 @@
+// Copyright Epic Games, Inc. All Rights Reserved.
+
+#pragma once
+
+#include "AdvancedRenamerSectionBase.h"
+#include "Input/Reply.h"
+
+class FString;
+class SWidget;
+
+/**
+ * Change case types
+ */
+enum class EAdvancedRenamerChangeCaseType : uint8
+{
+	SwapFirst,
+	SwapAll,
+	AllLower,
+	AllUpper
+};
+
+/**
+ * ChangeCase section
+ */
+class FAdvancedRenamerChangeCaseSection : public FAdvancedRenamerSectionBase
+{
+public:
+	FAdvancedRenamerChangeCaseSection();
+
+	virtual ~FAdvancedRenamerChangeCaseSection() {}
+
+	/** Init the given section */
+	virtual void Init(TSharedRef<IAdvancedRenamer> InRenamer) override;
+
+	/** Return the widget for this section*/
+	virtual TSharedRef<SWidget> GetWidget() override;
+
+	/** Reset all values of the section to the default ones */
+	virtual void ResetToDefault() override;
+
+private:
+	/** Called when one of the CaseChange buttons is clicked */
+	FReply OnChangeCaseButtonClicked(EAdvancedRenamerChangeCaseType InNewValue);
+
+	/** Whether or not the ChangeCase operation can be executed */
+	bool CanApplyChangeCaseSection();
+
+	/** Execute logic for the SwapFirst */
+	void ApplySwapFirst(FString& OutOriginalName);
+
+	/** Execute logic for the SwapAll */
+	void ApplySwapAll(FString& OutOriginalName);
+
+	/** Execute logic for the AllLower */
+	void ApplyAllLower(FString& OutOriginalName);
+
+	/** Execute logic for the AllUpper */
+	void ApplyAllUpper(FString& OutOriginalName);
+
+	/** Reset the current button clicked and disable the ChangeCase until the next time a button is clicked */
+	void ResetButtonClicked();
+
+	/** Execute logic for this section */
+	void ApplyChangeCaseSection(FString& OutOriginalName);
+
+private:
+	/** ChangeCase type to execute */
+	EAdvancedRenamerChangeCaseType ChangeCaseType;
+	
+	/** Whether a button was clicked and the Renamer needs to apply the ChangeCase */
+	bool bButtonWasClicked;
+};
diff --git a/Engine/Plugins/Experimental/AdvancedRenamer/Source/AdvancedRenamer/Private/AdvancedRenamerSections/AdvancedRenamerRemovePrefixSection.cpp b/Engine/Plugins/Experimental/AdvancedRenamer/Source/AdvancedRenamer/Private/AdvancedRenamerSections/AdvancedRenamerRemovePrefixSection.cpp
new file mode 100644
--- /dev/null
+++ b/Engine/Plugins/Experimental/AdvancedRenamer/Source/AdvancedRenamer/Private/AdvancedRenamerSections/AdvancedRenamerRemovePrefixSection.cpp
@@ -0,0 +1,255 @@
+// Copyright Epic Games, Inc. All Rights Reserved.
+
+#include "AdvancedRenamerRemovePrefixSection.h"
+
+#include "IAdvancedRenamer.h"
+#include "AdvancedRenamerStyle.h"
+#include "Styling/StyleColors.h"
+#include "Widgets/Input/SEditableTextBox.h"
+#include "Widgets/Input/SSegmentedControl.h"
+#include "Widgets/Input/SSpinBox.h"
+#include "Widgets/Layout/SWidgetSwitcher.h"
+#include "Widgets/Layout/SBorder.h"
+#include "Widgets/Layout/SBox.h"
+
+#define LOCTEXT_NAMESPACE "AdvancedRenamerRemovePrefixSection"
+
+FAdvancedRenamerRemovePrefixSection::FAdvancedRenamerRemovePrefixSection()
+{
+	FAdvancedRenamerRemovePrefixSection::ResetToDefault();
+}
+
+void FAdvancedRenamerRemovePrefixSection::Init(TSharedRef<IAdvancedRenamer> InRenamer)
+{
+	FAdvancedRenamerSectionBase::Init(InRenamer);
+	Section.SectionName = TEXT("RemovePrefix");
+	Section.OnOperationExecuted().BindSP(this, &FAdvancedRenamerRemovePrefixSection::ApplyRemovePrefixOperation);
+	InRenamer->AddSection(Section);
+}
+
+TSharedRef<SWidget> FAdvancedRenamerRemovePrefixSection::GetWidget()
+{
+	using namespace AdvancedRenamerSlateUtils::Default;
+
+	return SNew(SBorder)
+		.BorderImage(FAdvancedRenamerStyle::Get().GetBrush("AdvancedRenamer.Style.BackgroundBorder"))
+		.Content()
+		[
+			SNew(SVerticalBox)
+
+			// Remove Prefix CheckBox
+			+ SVerticalBox::Slot()
+			.Padding(SectionContentFirstEntryPadding)
+			.AutoHeight()
+			[
+				SNew(SHorizontalBox)
+
+				+ SHorizontalBox::Slot()
+				.HAlign(HAlign_Left)
+				.VAlign(VAlign_Center)
+				.AutoWidth()
+				.Padding(NameWidgetPadding)
+				[
+					SAssignNew(RemoveOldPrefixCheckBox, SCheckBox)
+					.IsChecked(this, &FAdvancedRenamerRemovePrefixSection::IsRemoveOldPrefixChecked)
+					.OnCheckStateChanged(this, &FAdvancedRenamerRemovePrefixSection::OnRemoveOldPrefixCheckBoxChanged)
+				]
+				+ SHorizontalBox::Slot()
+				.HAlign(HAlign_Left)
+				.VAlign(VAlign_Center)
+				.AutoWidth()
+				.Padding(VerticalAddNumberPadding)
+				[
+					SNew(STextBlock)
+					.Font(RegularFont)
+					.Text(LOCTEXT("AR_RemoveOldPrefix", "Remove Old Prefix"))
+				]
+			]
+
+			// Remove Prefix
+			+ SVerticalBox::Slot()
+			.Padding(SectionContentMiddleEntriesPadding)
+			.AutoHeight()
+			[
+				SNew(SHorizontalBox)
+				
+				+ SHorizontalBox::Slot()
+				.HAlign(HAlign_Left)
+				.VAlign(VAlign_Center)
+				.AutoWidth()
+				.Padding(NameWidgetPadding)
+				[
+					SNew(SBox)
+					.HeightOverride(25.f)
+					.MinDesiredWidth(160.f)
+					[
+						SNew(SSegmentedControl<EAdvancedRenamerRemoveOldType>)
+						.TextStyle(&FAppStyle::Get().GetWidgetStyle<FTextBlockStyle>("NormalText"))
+						.IsEnabled(this, &FAdvancedRenamerRemovePrefixSection::IsRemoveOldPrefixEnabled)
+						.UniformPadding(FMargin(2.f, 0.f))
+						.SupportsEmptySelection(false)
+						.SupportsMultiSelection(false)
+						.Value(this, &FAdvancedRenamerRemovePrefixSection::GetRemoveOldPrefixType)
+						.OnValueChanged(this, &FAdvancedRenamerRemovePrefixSection::OnRemoveOldPrefixTypeChanged)
+
+						+ SSegmentedControl<EAdvancedRenamerRemoveOldType>::Slot(EAdvancedRenamerRemoveOldType::Separator)
+						.Text(LOCTEXT("AR_SuffixSeparator", "Separator"))
+
+						+ SSegmentedControl<EAdvancedRenamerRemoveOldType>::Slot(EAdvancedRenamerRemoveOldType::Chars)
+						.Text(LOCTEXT("AR_SuffixFirstChar(s)", "First Char(s)"))
+					]
+				]
+				+ SHorizontalBox::Slot()
+				.Padding(ValueWidgetPadding)
+				[
+					SNew(SWidgetSwitcher)
+					.IsEnabled(this, &FAdvancedRenamerRemovePrefixSection::IsRemoveOldPrefixEnabled)
+					.WidgetIndex(TAttribute<int32>::CreateLambda([this] () { return PrefixWidgetSwitcherIndex; }))
+
+					+ SWidgetSwitcher::Slot()
+					[
+						SAssignNew(PrefixSeparatorTextBox, SEditableTextBox)
+						.BackgroundColor(FStyleColors::Background.GetSpecifiedColor())
+						.Font(RegularFont)
+						.Text(this, &FAdvancedRenamerRemovePrefixSection::GetPrefixSeparatorText)
+						.OnVerifyTextChanged(this, &FAdvancedRenamerRemovePrefixSection::OnPrefixSeparatorVerifyTextChanged)
+						.OnTextChanged(this, &FAdvancedRenamerRemovePrefixSection::OnPrefixSeparatorChanged)
+					]
+
+					+ SWidgetSwitcher::Slot()
+					[
+						SAssignNew(PrefixRemoveCharactersSpinBox, SSpinBox<uint8>)
+						.Style(&FAppStyle::Get(), "Menu.SpinBox")
+						.Font(RegularFont)
+						.MinValue(1)
+						.MaxValue(9)
+						.Value(this, &FAdvancedRenamerRemovePrefixSection::GetPrefixCharsValue)
+						.OnValueChanged(this, &FAdvancedRenamerRemovePrefixSection::OnPrefixRemoveCharactersChanged)
+					]
+				]
+			]
+		];
+}
+
+void FAdvancedRenamerRemovePrefixSection::ResetToDefault()
+{
+	bRemoveOldPrefixSection = false;
+	PrefixWidgetSwitcherIndex = 0;
+	RemovePrefixCharsValue = 1;
+	RemovePrefixSeparatorText = LOCTEXT("AR_PrefixSeparator", "_");
+	RemovePrefixType = EAdvancedRenamerRemoveOldType::Separator;
+}
+
+ECheckBoxState FAdvancedRenamerRemovePrefixSection::IsRemoveOldPrefixChecked() const
+{
+	return bRemoveOldPrefixSection ? ECheckBoxState::Checked : ECheckBoxState::Unchecked;
+}
+
+bool FAdvancedRenamerRemovePrefixSection::IsRemoveOldPrefixEnabled() const
+{
+	return bRemoveOldPrefixSection;
+}
+
+void FAdvancedRenamerRemovePrefixSection::OnRemoveOldPrefixCheckBoxChanged(ECheckBoxState InNewState)
+{
+	bRemoveOldPrefixSection = InNewState == ECheckBoxState::Checked;
+	MarkRenamerDirty();
+}
+
+FText FAdvancedRenamerRemovePrefixSection::GetPrefixSeparatorText() const
+{
+	return RemovePrefixSeparatorText;
+}
+
+uint8 FAdvancedRenamerRemovePrefixSection::GetPrefixCharsValue() const
+{
+	return RemovePrefixCharsValue;
+}
+
+EAdvancedRenamerRemoveOldType FAdvancedRenamerRemovePrefixSection::GetRemoveOldPrefixType() const
+{
+	return RemovePrefixType;
+}
+
+bool FAdvancedRenamerRemovePrefixSection::OnPrefixSeparatorVerifyTextChanged(const FText& InText, FText& OutErrorText) const
+{
+	if (InText.ToString().Len() > 1)
+	{
+		OutErrorText = LOCTEXT("SeparatorError", "Separators can only be a single character.");
+		return false;
+	}
+
+	return true;
+}
+
+void FAdvancedRenamerRemovePrefixSection::OnPrefixSeparatorChanged(const FText& InNewText)
+{
+	RemovePrefixSeparatorText = InNewText;
+	MarkRenamerDirty();
+}
+
+void FAdvancedRenamerRemovePrefixSection::OnPrefixRemoveCharactersChanged(uint8 InNewValue)
+{
+	RemovePrefixCharsValue = InNewValue;
+	MarkRenamerDirty();
+}
+
+void FAdvancedRenamerRemovePrefixSection::OnRemoveOldPrefixTypeChanged(EAdvancedRenamerRemoveOldType InNewValue)
+{
+	RemovePrefixType = InNewValue;
+	switch (InNewValue)
+	{
+	case EAdvancedRenamerRemoveOldType::Separator:
+		PrefixWidgetSwitcherIndex = 0;
+		break;
+	case EAdvancedRenamerRemoveOldType::Chars:
+		PrefixWidgetSwitcherIndex = 1;
+		break;
+	default:
+		break;
+	}
+	MarkRenamerDirty();
+}
+
+bool FAdvancedRenamerRemovePrefixSection::CanApplyRemovePrefixSeparatorOperation()
+{
+	return RemovePrefixType == EAdvancedRenamerRemoveOldType::Separator && RemovePrefixSeparatorText.ToString().Len() == 1;
+}
+
+bool FAdvancedRenamerRemovePrefixSection::CanApplyRemovePrefixCharOperation()
+{
+	return RemovePrefixType == EAdvancedRenamerRemoveOldType::Chars;
+}
+
+void FAdvancedRenamerRemovePrefixSection::ApplyRemovePrefixSeparatorOperation(FString& OutOriginalName)
+{
+	FString Separator = RemovePrefixSeparatorText.ToString();
+	const int32 PrefixStart = OutOriginalName.Find(Separator, ESearchCase::IgnoreCase);
+
+	if (PrefixStart >= 0)
+	{
+		OutOriginalName.RightChopInline(PrefixStart + Separator.Len());
+	}
+}
+
+void FAdvancedRenamerRemovePrefixSection::ApplyRemovePrefixCharOperation(FString& OutOriginalName)
+{
+	OutOriginalName.RightChopInline(RemovePrefixCharsValue);
+}
+
+void FAdvancedRenamerRemovePrefixSection::ApplyRemovePrefixOperation(FString& OutOriginalName)
+{
+	if (bRemoveOldPrefixSection)
+	{
+		if (CanApplyRemovePrefixSeparatorOperation())
+		{
+			ApplyRemovePrefixSeparatorOperation(OutOriginalName);
+		}
+		else if (CanApplyRemovePrefixCharOperation())
+		{
+			ApplyRemovePrefixCharOperation(OutOriginalName);
+		}
+	}
+}
+
+#undef LOCTEXT_NAMESPACE
diff --git a/Engine/Plugins/Experimental/AdvancedRenamer/Source/AdvancedRenamer/Private/AdvancedRenamerSections/AdvancedRenamerRemovePrefixSection.h b/Engine/Plugins/Experimental/AdvancedRenamer/Source/AdvancedRenamer/Private/AdvancedRenamerSections/AdvancedRenamerRemovePrefixSection.h
new file mode 100644
--- /dev/null
+++ b/Engine/Plugins/Experimental/AdvancedRenamer/Source/AdvancedRenamer/Private/AdvancedRenamerSections/AdvancedRenamerRemovePrefixSection.h
@@ -0,0 +1,102 @@
+// Copyright Epic Games, Inc. All Rights Reserved.
+
+#pragma once
+
+#include "AdvancedRenamerSectionBase.h"
+#include "Styling/SlateTypes.h"
+#include "Utils/AdvancedRenamerSlateUtils.h"
+
+class FString;
+class SCheckBox;
+class SEditableTextBox;
+template <typename InNumericType>
+class SSpinBox;
+class SWidget;
+
+class FAdvancedRenamerRemovePrefixSection : public FAdvancedRenamerSectionBase
+{
+public:
+	FAdvancedRenamerRemovePrefixSection();
+
+	virtual ~FAdvancedRenamerRemovePrefixSection() {}
+
+	/** Init the given section */
+	virtual void Init(TSharedRef<IAdvancedRenamer> InRenamer) override;
+
+	/** Return the widget for this section*/
+	virtual TSharedRef<SWidget> GetWidget() override;
+
+	/** Reset all values of the section to the default ones */
+	virtual void ResetToDefault() override;
+
+private:
+	/** If checked the section will be enabled */
+	ECheckBoxState IsRemoveOldPrefixChecked() const;
+
+	/** If checked it will remove all the current suffix numbers */
+	bool IsRemoveOldPrefixEnabled() const;
+
+	/** Get the separator text */
+	FText GetPrefixSeparatorText() const;
+
+	/** Get the numbers of chars to remove */
+	uint8 GetPrefixCharsValue() const;
+
+	/** Get the type of the remove, either by Separator or by number of Chars */
+	EAdvancedRenamerRemoveOldType GetRemoveOldPrefixType() const;
+
+	/** Called when the RemoveOldPrefix checkbox state change */
+	void OnRemoveOldPrefixCheckBoxChanged(ECheckBoxState InNewState);
+
+	/** Verify that the separator text is valid */
+	bool OnPrefixSeparatorVerifyTextChanged(const FText& InText, FText& OutErrorText) const;
+
+	/** Called when the separator text change */
+	void OnPrefixSeparatorChanged(const FText& InNewText);
+
+	/** Called when the remove chars number change */
+	void OnPrefixRemoveCharactersChanged(uint8 InNewValue);
+
+	/** Called when the remove type changed, it will be either by Separator or by number of Chars */
+	void OnRemoveOldPrefixTypeChanged(EAdvancedRenamerRemoveOldType InNewValue);
+
+	/** Whether or not the Remove by Separator operation can be executed */
+	bool CanApplyRemovePrefixSeparatorOperation();
+
+	/** Whether or not the Remove by number of Chars operation can be executed */
+	bool CanApplyRemovePrefixCharOperation();
+
+	/** Execute logic for the Remove by Separator */
+	void ApplyRemovePrefixSeparatorOperation(FString& OutOriginalName);
+
+	/** Execute logic for the Remove by number of Chars */
+	void ApplyRemovePrefixCharOperation(FString& OutOriginalName);
+
+	/** Execute logic for this section */
+	void ApplyRemovePrefixOperation(FString& OutOriginalName);
+
+private:
+	/** Section enabler CheckBox */
+	TSharedPtr<SCheckBox> RemoveOldPrefixCheckBox;
+
+	/** Separator EditableTextBox */
+	TSharedPtr<SEditableTextBox> PrefixSeparatorTextBox;
+
+	/** Number of Chars SSpinBox */
+	TSharedPtr<SSpinBox<uint8>> PrefixRemoveCharactersSpinBox;
+	
+	/** Number of Chars value */
+	uint8 RemovePrefixCharsValue;
+
+	/** Separator Text */
+	FText RemovePrefixSeparatorText;
+	
+	/** Index of the WidgetSwitcher */
+	int32 PrefixWidgetSwitcherIndex;
+
+	/** RemoveSuffix Remove type, either Separator or by number of Chars */
+	EAdvancedRenamerRemoveOldType RemovePrefixType;
+
+	/** RemoveOldSuffix section enabler, section is enabled if true */
+	bool bRemoveOldPrefixSection;
+};
diff --git a/Engine/Plugins/Experimental/AdvancedRenamer/Source/AdvancedRenamer/Private/AdvancedRenamerSections/AdvancedRenamerRemoveSuffixSection.cpp b/Engine/Plugins/Experimental/AdvancedRenamer/Source/AdvancedRenamer/Private/AdvancedRenamerSections/AdvancedRenamerRemoveSuffixSection.cpp
new file mode 100644
--- /dev/null
+++ b/Engine/Plugins/Experimental/AdvancedRenamer/Source/AdvancedRenamer/Private/AdvancedRenamerSections/AdvancedRenamerRemoveSuffixSection.cpp
@@ -0,0 +1,327 @@
+// Copyright Epic Games, Inc. All Rights Reserved.
+
+#include "AdvancedRenamerRemoveSuffixSection.h"
+
+#include "IAdvancedRenamer.h"
+#include "AdvancedRenamerStyle.h"
+#include "Styling/StyleColors.h"
+#include "Widgets/Input/SEditableTextBox.h"
+#include "Widgets/Input/SSegmentedControl.h"
+#include "Widgets/Input/SSpinBox.h"
+#include "Widgets/Layout/SWidgetSwitcher.h"
+#include "Widgets/Layout/SBorder.h"
+#include "Widgets/Layout/SBox.h"
+
+#define LOCTEXT_NAMESPACE "AdvancedRenamerRemoveSuffixSection"
+
+FAdvancedRenamerRemoveSuffixSection::FAdvancedRenamerRemoveSuffixSection()
+{
+	FAdvancedRenamerRemoveSuffixSection::ResetToDefault();
+}
+
+void FAdvancedRenamerRemoveSuffixSection::Init(TSharedRef<IAdvancedRenamer> InRenamer)
+{
+	FAdvancedRenamerSectionBase::Init(InRenamer);
+	Section.SectionName = TEXT("RemoveSuffix");
+	Section.OnOperationExecuted().BindSP(this, &FAdvancedRenamerRemoveSuffixSection::ApplyRemoveSuffixOperation);
+	InRenamer->AddSection(Section);
+}
+
+TSharedRef<SWidget> FAdvancedRenamerRemoveSuffixSection::GetWidget()
+{
+	using namespace AdvancedRenamerSlateUtils::Default;
+
+	return SNew(SBorder)
+		.BorderImage(FAdvancedRenamerStyle::Get().GetBrush("AdvancedRenamer.Style.BackgroundBorder"))
+		.Content()
+		[
+			SNew(SVerticalBox)
+
+			// Remove Suffix CheckBox
+			+ SVerticalBox::Slot()
+			.Padding(SectionContentFirstEntryPadding)
+			.AutoHeight()
+			[
+				SNew(SHorizontalBox)
+
+				+ SHorizontalBox::Slot()
+				.HAlign(HAlign_Left)
+				.VAlign(VAlign_Center)
+				.AutoWidth()
+				.Padding(NameWidgetPadding)
+				[
+					SAssignNew(RemoveOldSuffixCheckBox, SCheckBox)
+					.IsChecked(this, &FAdvancedRenamerRemoveSuffixSection::IsRemoveOldSuffixChecked)
+					.OnCheckStateChanged(this, &FAdvancedRenamerRemoveSuffixSection::OnRemoveOldSuffixCheckBoxChanged)
+				]
+				+ SHorizontalBox::Slot()
+				.HAlign(HAlign_Left)
+				.VAlign(VAlign_Center)
+				.AutoWidth()
+				.Padding(VerticalAddNumberPadding)
+				[
+					SNew(STextBlock)
+					.Font(RegularFont)
+					.Text(LOCTEXT("AR_RemoveOldSuffix", "Remove Old Suffix"))
+				]
+			]
+
+			// Remove Suffix
+			+ SVerticalBox::Slot()
+			.Padding(SectionContentMiddleEntriesPadding)
+			.AutoHeight()
+			[
+				SNew(SHorizontalBox)
+
+				+ SHorizontalBox::Slot()
+				.HAlign(HAlign_Left)
+				.VAlign(VAlign_Center)
+				.AutoWidth()
+				.Padding(NameWidgetPadding)
+				[
+					SNew(SBox)
+					.HeightOverride(25.f)
+					.MinDesiredWidth(160.f)
+					[
+						SNew(SSegmentedControl<EAdvancedRenamerRemoveOldType>)
+						.TextStyle(&FAppStyle::Get().GetWidgetStyle<FTextBlockStyle>("NormalText"))
+						.IsEnabled(this, &FAdvancedRenamerRemoveSuffixSection::IsRemoveOldSuffixEnabled)
+						.UniformPadding(FMargin(2.f, 0.f))
+						.SupportsEmptySelection(false)
+						.SupportsMultiSelection(false)
+						.Value(this, &FAdvancedRenamerRemoveSuffixSection::GetRemoveOldSuffixType)
+						.OnValueChanged(this, &FAdvancedRenamerRemoveSuffixSection::OnRemoveOldSuffixTypeChanged)
+
+						+ SSegmentedControl<EAdvancedRenamerRemoveOldType>::Slot(EAdvancedRenamerRemoveOldType::Separator)
+						.Text(LOCTEXT("AR_SuffixSeparator", "Separator"))
+
+						+ SSegmentedControl<EAdvancedRenamerRemoveOldType>::Slot(EAdvancedRenamerRemoveOldType::Chars)
+						.Text(LOCTEXT("AR_SuffixFirstChar(s)", "Last Char(s)"))
+					]
+				]
+				+ SHorizontalBox::Slot()
+				.Padding(ValueWidgetPadding)
+				[
+					SNew(SWidgetSwitcher)
+					.IsEnabled(this, &FAdvancedRenamerRemoveSuffixSection::IsRemoveOldSuffixEnabled)
+					.WidgetIndex(TAttribute<int32>::CreateLambda([this] () { return SuffixWidgetSwitcherIndex; }))
+
+					+ SWidgetSwitcher::Slot()
+					[
+						SAssignNew(SuffixSeparatorTextBox, SEditableTextBox)
+						.BackgroundColor(FStyleColors::Background.GetSpecifiedColor())
+						.Font(RegularFont)
+						.Text(this, &FAdvancedRenamerRemoveSuffixSection::GetSuffixSeparatorText)
+						.OnVerifyTextChanged(this, &FAdvancedRenamerRemoveSuffixSection::OnSuffixSeparatorVerifyTextChanged)
+						.OnTextChanged(this, &FAdvancedRenamerRemoveSuffixSection::OnSuffixSeparatorChanged)
+					]
+
+					+ SWidgetSwitcher::Slot()
+					[
+						SAssignNew(SuffixRemoveCharactersSpinBox, SSpinBox<uint8>)
+						.Style(&FAppStyle::Get(), "Menu.SpinBox")
+						.Font(RegularFont)
+						.MinValue(1)
+						.MaxValue(9)
+						.Value(this, &FAdvancedRenamerRemoveSuffixSection::GetSuffixCharsValue)
+						.OnValueChanged(this, &FAdvancedRenamerRemoveSuffixSection::OnSuffixRemoveCharactersChanged)
+					]
+				]
+			]
+
+			// Remove Suffix Numbers
+			+ SVerticalBox::Slot()
+			.Padding(SectionContentMiddleEntriesPadding)
+			.AutoHeight()
+			[
+				SNew(SHorizontalBox)
+
+				+ SHorizontalBox::Slot()
+				.HAlign(HAlign_Left)
+				.VAlign(VAlign_Center)
+				.Padding(NameWidgetPadding)
+				.AutoWidth()
+				[
+					SAssignNew(SuffixRemoveNumberCheckBox, SCheckBox)
+					.IsEnabled(this, &FAdvancedRenamerRemoveSuffixSection::IsRemoveOldSuffixEnabled)
+					.IsChecked(this, &FAdvancedRenamerRemoveSuffixSection::IsSuffixRemoveNumberChecked)
+					.OnCheckStateChanged(this, &FAdvancedRenamerRemoveSuffixSection::OnSuffixRemoveNumberCheckBoxChanged)
+				]
+				+ SHorizontalBox::Slot()
+				.HAlign(HAlign_Left)
+				.VAlign(VAlign_Center)
+				.Padding(VerticalAddNumberPadding)
+				.AutoWidth()
+				[
+					SNew(STextBlock)
+					.Font(RegularFont)
+					.Text(LOCTEXT("AR_SuffixRemoveNumber", "Remove Number"))
+				]
+			]
+		];
+}
+
+void FAdvancedRenamerRemoveSuffixSection::ResetToDefault()
+{
+	bRemoveOldSuffixSection = false;
+	bRemoveSuffixNumbers = false;
+	SuffixWidgetSwitcherIndex = 0;
+	RemoveSuffixCharsValue = 1;
+	RemoveSuffixSeparatorText = LOCTEXT("AR_SuffixSeparator", "_");
+	RemoveSuffixType = EAdvancedRenamerRemoveOldType::Separator;
+}
+
+ECheckBoxState FAdvancedRenamerRemoveSuffixSection::IsRemoveOldSuffixChecked() const
+{
+	return bRemoveOldSuffixSection ? ECheckBoxState::Checked : ECheckBoxState::Unchecked;
+}
+
+ECheckBoxState FAdvancedRenamerRemoveSuffixSection::IsSuffixRemoveNumberChecked() const
+{
+	return bRemoveSuffixNumbers ? ECheckBoxState::Checked : ECheckBoxState::Unchecked;
+}
+
+bool FAdvancedRenamerRemoveSuffixSection::IsRemoveOldSuffixEnabled() const
+{
+	return bRemoveOldSuffixSection;
+}
+
+void FAdvancedRenamerRemoveSuffixSection::OnRemoveOldSuffixCheckBoxChanged(ECheckBoxState InNewState)
+{
+	bRemoveOldSuffixSection = InNewState == ECheckBoxState::Checked;
+	MarkRenamerDirty();
+}
+
+void FAdvancedRenamerRemoveSuffixSection::OnSuffixRemoveNumberCheckBoxChanged(ECheckBoxState InNewState)
+{
+	bRemoveSuffixNumbers = InNewState == ECheckBoxState::Checked;
+	MarkRenamerDirty();
+}
+
+FText FAdvancedRenamerRemoveSuffixSection::GetSuffixSeparatorText() const
+{
+	return RemoveSuffixSeparatorText;
+}
+
+uint8 FAdvancedRenamerRemoveSuffixSection::GetSuffixCharsValue() const
+{
+	return RemoveSuffixCharsValue;
+}
+
+EAdvancedRenamerRemoveOldType FAdvancedRenamerRemoveSuffixSection::GetRemoveOldSuffixType() const
+{
+	return RemoveSuffixType;
+}
+
+bool FAdvancedRenamerRemoveSuffixSection::OnSuffixSeparatorVerifyTextChanged(const FText& InText, FText& OutErrorText) const
+{
+	if (InText.ToString().Len() > 1)
+	{
+		OutErrorText = LOCTEXT("SeparatorError", "Separators can only be a single character.");
+		return false;
+	}
+
+	return true;
+}
+
+void FAdvancedRenamerRemoveSuffixSection::OnSuffixSeparatorChanged(const FText& InNewText)
+{
+	RemoveSuffixSeparatorText = InNewText;
+	MarkRenamerDirty();
+}
+
+void FAdvancedRenamerRemoveSuffixSection::OnSuffixRemoveCharactersChanged(uint8 InNewValue)
+{
+	RemoveSuffixCharsValue = InNewValue;
+	MarkRenamerDirty();
+}
+
+void FAdvancedRenamerRemoveSuffixSection::OnRemoveOldSuffixTypeChanged(EAdvancedRenamerRemoveOldType InNewValue)
+{
+	RemoveSuffixType = InNewValue;
+	switch (InNewValue)
+	{
+		case EAdvancedRenamerRemoveOldType::Separator:
+			SuffixWidgetSwitcherIndex = 0;
+			break;
+		case EAdvancedRenamerRemoveOldType::Chars:
+			SuffixWidgetSwitcherIndex = 1;
+			break;
+		default:
+			break;
+	}
+	MarkRenamerDirty();
+}
+
+bool FAdvancedRenamerRemoveSuffixSection::CanApplyRemoveSuffixSeparatorOperation()
+{
+	return RemoveSuffixType == EAdvancedRenamerRemoveOldType::Separator && RemoveSuffixSeparatorText.ToString().Len() == 1;
+}
+
+bool FAdvancedRenamerRemoveSuffixSection::CanApplyRemoveSuffixCharOperation()
+{
+	return RemoveSuffixType == EAdvancedRenamerRemoveOldType::Chars;
+}
+
+bool FAdvancedRenamerRemoveSuffixSection::CanApplyRemoveSuffixNumbers()
+{
+	return bRemoveSuffixNumbers;
+}
+
+void FAdvancedRenamerRemoveSuffixSection::ApplyRemoveSuffixSeparatorOperation(FString& OutOriginalName)
+{
+	FString Separator = RemoveSuffixSeparatorText.ToString();
+	const int32 SuffixStart = OutOriginalName.Find(Separator, ESearchCase::IgnoreCase, ESearchDir::FromEnd);
+
+	if (SuffixStart >= 0)
+	{
+		OutOriginalName.LeftInline(SuffixStart);
+	}
+}
+
+void FAdvancedRenamerRemoveSuffixSection::ApplyRemoveSuffixCharOperation(FString& OutOriginalName)
+{
+	OutOriginalName.LeftChopInline(RemoveSuffixCharsValue);
+}
+
+void FAdvancedRenamerRemoveSuffixSection::ApplyRemoveSuffixNumbers(FString& OutOriginalName)
+{
+	static constexpr TCHAR FirstDigit = '0';
+	static constexpr TCHAR LastDigit = '9';
+
+	int32 LastDigitIndex = OutOriginalName.Len() - 1;
+
+	while (LastDigitIndex >= 0)
+	{
+		if (OutOriginalName[LastDigitIndex] < FirstDigit || OutOriginalName[LastDigitIndex] > LastDigit)
+		{
+			break;
+		}
+
+		--LastDigitIndex;
+	}
+
+	OutOriginalName.LeftInline(LastDigitIndex + 1);
+}
+
+void FAdvancedRenamerRemoveSuffixSection::ApplyRemoveSuffixOperation(FString& OutOriginalName)
+{
+	if (bRemoveOldSuffixSection)
+	{
+		if (CanApplyRemoveSuffixNumbers())
+		{
+			ApplyRemoveSuffixNumbers(OutOriginalName);
+		}
+
+		if (CanApplyRemoveSuffixSeparatorOperation())
+		{
+			ApplyRemoveSuffixSeparatorOperation(OutOriginalName);
+		}
+		else if (CanApplyRemoveSuffixCharOperation())
+		{
+			ApplyRemoveSuffixCharOperation(OutOriginalName);
+		}
+	}
+}
+
+#undef LOCTEXT_NAMESPACE
diff --git a/Engine/Plugins/Experimental/AdvancedRenamer/Source/AdvancedRenamer/Private/AdvancedRenamerSections/AdvancedRenamerRemoveSuffixSection.h b/Engine/Plugins/Experimental/AdvancedRenamer/Source/AdvancedRenamer/Private/AdvancedRenamerSections/AdvancedRenamerRemoveSuffixSection.h
new file mode 100644
--- /dev/null
+++ b/Engine/Plugins/Experimental/AdvancedRenamer/Source/AdvancedRenamer/Private/AdvancedRenamerSections/AdvancedRenamerRemoveSuffixSection.h
@@ -0,0 +1,124 @@
+// Copyright Epic Games, Inc. All Rights Reserved.
+
+#pragma once
+
+#include "AdvancedRenamerSectionBase.h"
+#include "Styling/SlateTypes.h"
+#include "Utils/AdvancedRenamerSlateUtils.h"
+
+class FString;
+class SCheckBox;
+class SEditableTextBox;
+template <typename InNumericType>
+class SSpinBox;
+class SWidget;
+
+/**
+ * RemoveSuffix section
+ */
+class FAdvancedRenamerRemoveSuffixSection : public FAdvancedRenamerSectionBase
+{
+public:
+	FAdvancedRenamerRemoveSuffixSection();
+
+	virtual ~FAdvancedRenamerRemoveSuffixSection() {}
+
+	/** Init the given section */
+	virtual void Init(TSharedRef<IAdvancedRenamer> InRenamer) override;
+
+	/** Return the widget for this section*/
+	virtual TSharedRef<SWidget> GetWidget() override;
+
+	/** Reset all values of the section to the default ones */
+	virtual void ResetToDefault() override;
+
+private:
+
+	/** If checked the section will be enabled */
+	ECheckBoxState IsRemoveOldSuffixChecked() const;
+
+	/** If checked it will remove all the current suffix numbers */
+	ECheckBoxState IsSuffixRemoveNumberChecked() const;
+
+	/** Whether or not the section is enabled */
+	bool IsRemoveOldSuffixEnabled() const;
+
+	/** Get the separator text */
+	FText GetSuffixSeparatorText() const;
+
+	/** Get the numbers of chars to remove */
+	uint8 GetSuffixCharsValue() const;
+
+	/** Get the type of the remove, either by Separator or by number of Chars */
+	EAdvancedRenamerRemoveOldType GetRemoveOldSuffixType() const;
+
+	/** Called when the RemoveOldSuffix checkbox state change */
+	void OnRemoveOldSuffixCheckBoxChanged(ECheckBoxState InNewState);
+
+	/** Called when the RemoveNumber checkbox state change */
+	void OnSuffixRemoveNumberCheckBoxChanged(ECheckBoxState InNewState);
+
+	/** Verify that the separator text is valid */
+	bool OnSuffixSeparatorVerifyTextChanged(const FText& InText, FText& OutErrorText) const;
+
+	/** Called when the separator text change */
+	void OnSuffixSeparatorChanged(const FText& InNewText);
+
+	/** Called when the remove chars number change */
+	void OnSuffixRemoveCharactersChanged(uint8 InNewValue);
+
+	/** Called when the remove type changed, it will be either by Separator or by number of Chars */
+	void OnRemoveOldSuffixTypeChanged(EAdvancedRenamerRemoveOldType InNewValue);
+
+	/** Whether or not the Remove by Separator operation can be executed */
+	bool CanApplyRemoveSuffixSeparatorOperation();
+
+	/** Whether or not the Remove by number of Chars operation can be executed */
+	bool CanApplyRemoveSuffixCharOperation();
+
+	/** Whether or not the RemoveNumbers operation can be executed */
+	bool CanApplyRemoveSuffixNumbers();
+
+	/** Execute logic for the Remove by Separator */
+	void ApplyRemoveSuffixSeparatorOperation(FString& OutOriginalName);
+
+	/** Execute logic for the Remove by number of Chars */
+	void ApplyRemoveSuffixCharOperation(FString& OutOriginalName);
+
+	/** Execute logic for the RemoveNumbers */
+	void ApplyRemoveSuffixNumbers(FString& OutOriginalName);
+
+	/** Execute logic for this section */
+	void ApplyRemoveSuffixOperation(FString& OutOriginalName);
+
+private:
+	/** Section enabler CheckBox */
+	TSharedPtr<SCheckBox> RemoveOldSuffixCheckBox;
+	
+	/** Separator EditableTextBox */
+	TSharedPtr<SEditableTextBox> SuffixSeparatorTextBox;
+
+	/** Number of Chars SSpinBox */
+	TSharedPtr<SSpinBox<uint8>> SuffixRemoveCharactersSpinBox;
+
+	/** RemoveNumber CheckBox */
+	TSharedPtr<SCheckBox> SuffixRemoveNumberCheckBox;
+
+	/** Number of Chars value */
+	uint8 RemoveSuffixCharsValue;
+
+	/** Separator Text */
+	FText RemoveSuffixSeparatorText;
+
+	/** Index of the WidgetSwitcher */
+	int32 SuffixWidgetSwitcherIndex;
+
+	/** RemoveSuffix Remove type, either Separator or by number of Chars */
+	EAdvancedRenamerRemoveOldType RemoveSuffixType;
+
+	/** RemoveOldSuffix section enabler, section is enabled if true */
+	bool bRemoveOldSuffixSection;
+
+	/** RemoveNumber enabler, RemoveNumber operation will be executed if true */
+	bool bRemoveSuffixNumbers;
+};
diff --git a/Engine/Plugins/Experimental/AdvancedRenamer/Source/AdvancedRenamer/Private/AdvancedRenamerSections/AdvancedRenamerSearchAndReplaceSection.cpp b/Engine/Plugins/Experimental/AdvancedRenamer/Source/AdvancedRenamer/Private/AdvancedRenamerSections/AdvancedRenamerSearchAndReplaceSection.cpp
new file mode 100644
--- /dev/null
+++ b/Engine/Plugins/Experimental/AdvancedRenamer/Source/AdvancedRenamer/Private/AdvancedRenamerSections/AdvancedRenamerSearchAndReplaceSection.cpp
@@ -0,0 +1,362 @@
+// Copyright Epic Games, Inc. All Rights Reserved.
+
+#include "AdvancedRenamerSearchAndReplaceSection.h"
+
+#include "AdvancedRenamerModule.h"
+#include "AdvancedRenamerStyle.h"
+#include "IAdvancedRenamer.h"
+#include "Internationalization/Regex.h"
+#include "Styling/StyleColors.h"
+#include "Utils/AdvancedRenamerSlateUtils.h"
+#include "Widgets/Input/SEditableTextBox.h"
+#include "Widgets/Input/SSegmentedControl.h"
+#include "Widgets/Layout/SBorder.h"
+#include "Widgets/Layout/SBox.h"
+
+#define LOCTEXT_NAMESPACE "AdvancedRenamerSearchAndReplaceSection"
+
+FAdvancedRenamerSearchAndReplaceSection::FAdvancedRenamerSearchAndReplaceSection()
+{
+	FAdvancedRenamerSearchAndReplaceSection::ResetToDefault();
+}
+
+void FAdvancedRenamerSearchAndReplaceSection::Init(TSharedRef<IAdvancedRenamer> InRenamer)
+{
+	FAdvancedRenamerSectionBase::Init(InRenamer);
+	Section.SectionName = TEXT("SearchAndReplaceSection");
+	Section.OnOperationExecuted().BindSP(this, &FAdvancedRenamerSearchAndReplaceSection::ApplySearchAndReplaceOrRenameOperation);
+	InRenamer->AddSection(Section);
+}
+
+TSharedRef<SWidget> FAdvancedRenamerSearchAndReplaceSection::GetWidget()
+{
+	using namespace AdvancedRenamerSlateUtils::Default;
+
+	return SNew(SBorder)
+		.BorderImage(FAdvancedRenamerStyle::Get().GetBrush("AdvancedRenamer.Style.BackgroundBorder"))
+		.Content()
+		[
+			SNew(SVerticalBox)
+
+			// Title
+			+ SVerticalBox::Slot()
+			.Padding(SectionContentFirstEntryPadding)
+			.AutoHeight()
+			[
+				SNew(STextBlock)
+				.Font(TitleFont)
+				.Text(LOCTEXT("AR_SearchReplaceTitle", "Search and Replace"))
+			]
+
+			+ SVerticalBox::Slot()
+			.Padding(SectionContentMiddleEntriesPadding)
+			.AutoHeight()
+			[
+				SNew(SHorizontalBox)
+
+				+ SHorizontalBox::Slot()
+				.HAlign(HAlign_Left)
+				.VAlign(VAlign_Center)
+				.AutoWidth()
+				.Padding(NameWidgetPadding)
+				[
+					SNew(SBox)
+					.HeightOverride(25.f)
+					.MinDesiredWidth(160.f)
+					[
+						SNew(SSegmentedControl<EAdvancedRenamerSearchAndReplaceType>)
+						.TextStyle(&FAppStyle::Get().GetWidgetStyle<FTextBlockStyle>("NormalText"))
+						.UniformPadding(FMargin(2.f, 0.f))
+						.SupportsEmptySelection(false)
+						.SupportsMultiSelection(false)
+						.Value(SearchAndReplaceType)
+						.OnValueChanged_Lambda([this] (EAdvancedRenamerSearchAndReplaceType InSearchAndReplaceType)
+						{
+							SearchAndReplaceType = InSearchAndReplaceType;
+							MarkRenamerDirty();
+						})
+
+						+ SSegmentedControl<EAdvancedRenamerSearchAndReplaceType>::Slot(EAdvancedRenamerSearchAndReplaceType::PlainText)
+						.Text(LOCTEXT("BR_PlainText", "Plain Text"))
+
+						+ SSegmentedControl<EAdvancedRenamerSearchAndReplaceType>::Slot(EAdvancedRenamerSearchAndReplaceType::RegularExpression)
+						.Text(LOCTEXT("BR_Regex", "Regex"))
+					]
+				]
+
+				// Regex Help
+				+ SHorizontalBox::Slot()
+				.AutoWidth()
+				.VAlign(VAlign_Center)
+				.Padding(FMargin(4.f, 0.f))
+				[
+					SNew(SImage)
+					.ColorAndOpacity(FSlateColor::UseForeground())
+					.Image(FAppStyle::GetBrush("Icons.Help"))
+					.OnMouseButtonDown(this, &FAdvancedRenamerSearchAndReplaceSection::OnRegexHelp)
+				]
+
+				// Ignore case
+				+ SHorizontalBox::Slot()
+				.VAlign(VAlign_Center)
+				.AutoWidth()
+				.Padding(FMargin(4.f, 0.f))
+				[
+					SAssignNew(SearchReplaceIgnoreCaseCheckBox, SCheckBox)
+					.IsChecked(this, &FAdvancedRenamerSearchAndReplaceSection::IsSearchReplaceIgnoreCaseChecked)
+					.OnCheckStateChanged(this, &FAdvancedRenamerSearchAndReplaceSection::OnSearchReplaceIgnoreCaseCheckBoxChanged)
+				]
+
+				+ SHorizontalBox::Slot()
+				.VAlign(VAlign_Center)
+				.AutoWidth()
+				[
+					SNew(STextBlock)
+					.Font(RegularFont)
+					.Text(LOCTEXT("AR_IgnoreCase", "Ignore Case"))
+				]
+			]
+
+			// Search text
+			+ SVerticalBox::Slot()
+			.Padding(SectionContentMiddleEntriesPadding)
+			.FillHeight(1.f)
+			[
+				SAssignNew(SearchReplaceSearchTextBox, SEditableTextBox)
+				.BackgroundColor(FStyleColors::Background.GetSpecifiedColor())
+				.Font(RegularFont)
+				.HintText(LOCTEXT("AR_RegexSearchHint", "Search"))
+				.Text(this, &FAdvancedRenamerSearchAndReplaceSection::GetSearchText)
+				.OnTextChanged(this, &FAdvancedRenamerSearchAndReplaceSection::OnSearchTextChanged)
+			]
+
+			+ SVerticalBox::Slot()
+			.Padding(SectionContentMiddleEntriesPadding)
+			.FillHeight(1.f)
+			[
+				SAssignNew(SearchReplaceReplaceTextBox, SEditableTextBox)
+				.BackgroundColor(FStyleColors::Background.GetSpecifiedColor())
+				.Font(RegularFont)
+				.HintText(LOCTEXT("AR_RegexReplaceHint", "Rename To"))
+				.Text(this, &FAdvancedRenamerSearchAndReplaceSection::GetReplaceText)
+				.OnTextChanged(this, &FAdvancedRenamerSearchAndReplaceSection::OnReplaceTextChanged)
+			]
+		];
+}
+
+void FAdvancedRenamerSearchAndReplaceSection::ResetToDefault()
+{
+	SearchCaseType = ESearchCase::IgnoreCase;
+	SearchAndReplaceType = EAdvancedRenamerSearchAndReplaceType::PlainText;
+	SearchText = FText::GetEmpty();
+	ReplaceText = FText::GetEmpty();
+}
+
+FReply FAdvancedRenamerSearchAndReplaceSection::OnRegexHelp(const FGeometry& InGeometry, const FPointerEvent& InPointerEvent)
+{
+	FPlatformProcess::LaunchURL(*GetRegexDocumentationURL(), nullptr, nullptr);
+	return FReply::Handled();
+}
+
+ECheckBoxState FAdvancedRenamerSearchAndReplaceSection::IsSearchReplaceIgnoreCaseChecked() const
+{
+	return SearchCaseType == ESearchCase::IgnoreCase ? ECheckBoxState::Checked : ECheckBoxState::Unchecked;
+}
+
+void FAdvancedRenamerSearchAndReplaceSection::OnSearchReplaceIgnoreCaseCheckBoxChanged(ECheckBoxState InNewState)
+{
+	SearchCaseType = InNewState == ECheckBoxState::Checked ? ESearchCase::IgnoreCase : ESearchCase::CaseSensitive;
+	if (const TSharedPtr<IAdvancedRenamer>& Renamer = RenamerWeakPtr.Pin())
+	{
+		Renamer->MarkDirty();
+	}
+}
+
+void FAdvancedRenamerSearchAndReplaceSection::OnSearchTextChanged(const FText& InNewText)
+{
+	SearchText = InNewText;
+	MarkRenamerDirty();
+}
+
+FText FAdvancedRenamerSearchAndReplaceSection::GetSearchText() const
+{
+	return SearchText;
+}
+
+void FAdvancedRenamerSearchAndReplaceSection::OnReplaceTextChanged(const FText& InNewText)
+{
+	ReplaceText = InNewText;
+	MarkRenamerDirty();
+}
+
+FText FAdvancedRenamerSearchAndReplaceSection::GetReplaceText() const
+{
+	return ReplaceText;
+}
+
+FString FAdvancedRenamerSearchAndReplaceSection::RegexReplace(const FString& InOriginalString, const FRegexPattern& InPattern, const FString& InReplaceString) const
+{
+	static const FString EscapeString = TEXT("\\");
+	static constexpr TCHAR EscapeChar = '\\';
+	static const FString GroupString = TEXT("$");
+	static constexpr TCHAR GroupChar = '$';
+	static constexpr TCHAR FirstDigit = '0';
+	static constexpr TCHAR LastDigit = '9';
+	static constexpr TCHAR NullChar = 0;
+
+	FRegexMatcher Matcher(InPattern, InOriginalString);
+
+	FString Output = "";
+	int32 StartCharIndex = 0;
+	bool bReadingGroupName = false;
+	int32 GroupIndex = INDEX_NONE;
+
+	while (Matcher.FindNext())
+	{
+		// Add on part of string after start/previous match
+		if (StartCharIndex != Matcher.GetMatchBeginning())
+		{
+			Output += InOriginalString.Mid(StartCharIndex, Matcher.GetMatchBeginning() - StartCharIndex);
+		}
+
+		bool bEscaped = false;
+
+		for (int32 CharIndex = 0; CharIndex <= InReplaceString.Len(); ++CharIndex)
+		{
+			const TCHAR& Char = CharIndex < InReplaceString.Len() ? InReplaceString[CharIndex] : NullChar;
+
+			if (bReadingGroupName)
+			{
+				// Build group index
+				if (Char >= FirstDigit && Char <= LastDigit)
+				{
+					if (GroupIndex == INDEX_NONE)
+					{
+						GroupIndex = 0;
+					}
+					else if (GroupIndex > 0)
+					{
+						GroupIndex *= 10;
+					}
+
+					const int32 NextDigit = static_cast<int32>(Char - FirstDigit);
+					GroupIndex += NextDigit;
+					continue;
+				}
+				// We've read a group index, add it to the output string
+				else if (GroupIndex > 0)
+				{
+					if (Matcher.GetCaptureGroupBeginning(GroupIndex) == INDEX_NONE)
+					{
+						UE_LOG(LogARP, Error, TEXT("Regex: Capture group does not exist %d."), GroupIndex);
+					}
+
+					Output += Matcher.GetCaptureGroup(GroupIndex);
+				}
+				// $0 matches the entire matched string
+				else if (GroupIndex == 0)
+				{
+					Output += InOriginalString.Mid(Matcher.GetMatchBeginning(), Matcher.GetMatchEnding() - Matcher.GetMatchBeginning());
+				}
+				// An unescaped $
+				else
+				{
+					UE_LOG(LogARP, Error, TEXT("Regex: Unescaped %s."), *GroupString);
+
+					Output += GroupString;
+				}
+
+				bReadingGroupName = false;
+				// Continue regular parsing of this character.
+			}
+
+			// Check for special chars
+			if (!bEscaped)
+			{
+				if (Char == EscapeChar)
+				{
+					bEscaped = true;
+					continue;
+				}
+
+				if (Char == GroupChar)
+				{
+					bReadingGroupName = true;
+					GroupIndex = INDEX_NONE;
+					continue;
+				}
+			}
+			else
+			{
+				// If the last char is a \ assume that it's not an escape char
+				if (Char == NullChar)
+				{
+					UE_LOG(LogARP, Error, TEXT("Regex: Unescaped %s."), *EscapeString);
+
+					Output += EscapeChar;
+				}
+			}
+
+			if (Char != NullChar)
+			{
+				Output += Char;
+			}
+
+			bEscaped = false;
+		}
+
+		StartCharIndex = Matcher.GetMatchEnding();
+	}
+
+	// Add on the end of the string after the last match
+	if (StartCharIndex < InOriginalString.Len())
+	{
+		Output += InOriginalString.Mid(StartCharIndex);
+	}
+
+	return Output;
+}
+
+bool FAdvancedRenamerSearchAndReplaceSection::CanApplySearchAndReplaceOperation()
+{
+	return !SearchText.IsEmpty();
+}
+
+bool FAdvancedRenamerSearchAndReplaceSection::CanApplyRenameOperation()
+{
+	return !ReplaceText.IsEmpty() && SearchText.IsEmpty();
+}
+
+void FAdvancedRenamerSearchAndReplaceSection::ApplySearchAndReplaceOperation(FString& OutOriginalName)
+{
+	if (SearchAndReplaceType == EAdvancedRenamerSearchAndReplaceType::PlainText)
+	{
+		OutOriginalName = OutOriginalName.Replace(*SearchText.ToString(), *ReplaceText.ToString(), SearchCaseType);
+	}
+	else
+	{
+		const FRegexPattern RegexPattern = FRegexPattern(SearchText.ToString(),
+			SearchCaseType == ESearchCase::IgnoreCase ? ERegexPatternFlags::CaseInsensitive : ERegexPatternFlags::None);
+
+		OutOriginalName = RegexReplace(OutOriginalName, RegexPattern, ReplaceText.ToString());
+	}
+}
+
+void FAdvancedRenamerSearchAndReplaceSection::ApplyRenameOperation(FString& OutOriginalName)
+{
+	OutOriginalName = ReplaceText.ToString();
+}
+
+void FAdvancedRenamerSearchAndReplaceSection::ApplySearchAndReplaceOrRenameOperation(FString& OutOriginalName)
+{
+	if (CanApplySearchAndReplaceOperation())
+	{
+		ApplySearchAndReplaceOperation(OutOriginalName);
+	}
+	else if (CanApplyRenameOperation())
+	{
+		ApplyRenameOperation(OutOriginalName);
+	}
+}
+
+#undef LOCTEXT_NAMESPACE
diff --git a/Engine/Plugins/Experimental/AdvancedRenamer/Source/AdvancedRenamer/Private/AdvancedRenamerSections/AdvancedRenamerSearchAndReplaceSection.h b/Engine/Plugins/Experimental/AdvancedRenamer/Source/AdvancedRenamer/Private/AdvancedRenamerSections/AdvancedRenamerSearchAndReplaceSection.h
new file mode 100644
--- /dev/null
+++ b/Engine/Plugins/Experimental/AdvancedRenamer/Source/AdvancedRenamer/Private/AdvancedRenamerSections/AdvancedRenamerSearchAndReplaceSection.h
@@ -0,0 +1,107 @@
+// Copyright Epic Games, Inc. All Rights Reserved.
+
+#pragma once
+
+#include "AdvancedRenamerSectionBase.h"
+#include "Input/Reply.h"
+#include "Styling/SlateTypes.h"
+
+class FRegexPattern;
+class SCheckBox;
+class SEditableTextBox;
+struct FGeometry;
+struct FPointerEvent;
+
+/**
+ * Search And Replace options
+ */
+enum class EAdvancedRenamerSearchAndReplaceType
+{
+	PlainText,
+	RegularExpression
+};
+
+/**
+ * Search and Replace/Rename section
+ */
+class FAdvancedRenamerSearchAndReplaceSection : public FAdvancedRenamerSectionBase
+{
+public:
+	FAdvancedRenamerSearchAndReplaceSection();
+
+	virtual ~FAdvancedRenamerSearchAndReplaceSection() {}
+
+	/** Init the given section */
+	virtual void Init(TSharedRef<IAdvancedRenamer> InRenamer) override;
+
+	/** Return the widget for this section*/
+	virtual TSharedRef<SWidget> GetWidget() override;
+
+	/** Reset all values of the section to the default ones */
+	virtual void ResetToDefault() override;
+
+private:
+	/** Regex documentation URL */
+	static FString GetRegexDocumentationURL() { return TEXT("https://unicode-org.github.io/icu/userguide/strings/regexp.html"); }
+
+	/** Called when the Regex help icon is clicked */
+	FReply OnRegexHelp(const FGeometry& InGeometry, const FPointerEvent& InPointerEvent);
+
+	/** If checked it will use the IgnoreCase, otherwise it is CaseSensitive */
+	ECheckBoxState IsSearchReplaceIgnoreCaseChecked() const;
+
+	/** Called when the ignore case checkbox state change */
+	void OnSearchReplaceIgnoreCaseCheckBoxChanged(ECheckBoxState InNewState);
+
+	/** Called when the search text change */
+	void OnSearchTextChanged(const FText& InNewText);
+
+	/** Get the search text */
+	FText GetSearchText() const;
+
+	/** Called when the replace text change */
+	void OnReplaceTextChanged(const FText& InNewText);
+
+	/** Get the replace text */
+	FText GetReplaceText() const;
+
+	/** Regex Replace logic */
+	FString RegexReplace(const FString& InOriginalString, const FRegexPattern& InPattern, const FString& InReplaceString) const;
+
+	/** Whether or not the SearchAndReplace operation can be executed */
+	bool CanApplySearchAndReplaceOperation();
+
+	/** Whether or not the Rename operation can be executed */
+	bool CanApplyRenameOperation();
+
+	/** Execute logic for the SearchAndReplace */
+	void ApplySearchAndReplaceOperation(FString& OutOriginalName);
+
+	/** Execute logic for the Rename */
+	void ApplyRenameOperation(FString& OutOriginalName);
+
+	/** Execute logic for this section */
+	void ApplySearchAndReplaceOrRenameOperation(FString& OutOriginalName);
+
+private:
+	/** Search EditableTextBox */
+	TSharedPtr<SEditableTextBox> SearchReplaceSearchTextBox;
+
+	/** Replace/Rename EditableTextBox */
+	TSharedPtr<SEditableTextBox> SearchReplaceReplaceTextBox;
+
+	/** Ignore case CheckBox */
+	TSharedPtr<SCheckBox> SearchReplaceIgnoreCaseCheckBox;
+
+	/** SearchAndReplace search type, either PlainText or RegularExpression */
+	EAdvancedRenamerSearchAndReplaceType SearchAndReplaceType;
+
+	/** SearchAndReplace case type, either IgnoreCase or CaseSensitive */
+	ESearchCase::Type SearchCaseType;
+
+	/** Search Text */
+	FText SearchText;
+
+	/** Replace/Rename Text */
+	FText ReplaceText;
+};
diff --git a/Engine/Plugins/Experimental/AdvancedRenamer/Source/AdvancedRenamer/Private/AdvancedRenamerSections/AdvancedRenamerSectionBase.cpp b/Engine/Plugins/Experimental/AdvancedRenamer/Source/AdvancedRenamer/Private/AdvancedRenamerSections/AdvancedRenamerSectionBase.cpp
new file mode 100644
--- /dev/null
+++ b/Engine/Plugins/Experimental/AdvancedRenamer/Source/AdvancedRenamer/Private/AdvancedRenamerSections/AdvancedRenamerSectionBase.cpp
@@ -0,0 +1,24 @@
+// Copyright Epic Games, Inc. All Rights Reserved.
+
+#include "AdvancedRenamerSectionBase.h"
+
+#include "IAdvancedRenamer.h"
+
+void FAdvancedRenamerSectionBase::Init(TSharedRef<IAdvancedRenamer> InRenamer)
+{
+	RenamerWeakPtr = InRenamer;
+	ResetToDefault();
+}
+
+FAdvancedRenamerExecuteSection FAdvancedRenamerSectionBase::GetSection() const
+{
+	return Section;
+}
+
+void FAdvancedRenamerSectionBase::MarkRenamerDirty() const
+{
+	if (const TSharedPtr<IAdvancedRenamer>& Renamer = RenamerWeakPtr.Pin())
+	{
+		Renamer->MarkDirty();
+	}
+}
diff --git a/Engine/Plugins/Experimental/AdvancedRenamer/Source/AdvancedRenamer/Private/AdvancedRenamerSections/AdvancedRenamerSectionBase.h b/Engine/Plugins/Experimental/AdvancedRenamer/Source/AdvancedRenamer/Private/AdvancedRenamerSections/AdvancedRenamerSectionBase.h
new file mode 100644
--- /dev/null
+++ b/Engine/Plugins/Experimental/AdvancedRenamer/Source/AdvancedRenamer/Private/AdvancedRenamerSections/AdvancedRenamerSectionBase.h
@@ -0,0 +1,29 @@
+// Copyright Epic Games, Inc. All Rights Reserved.
+
+#pragma once
+
+#include "IAdvancedRenamerSection.h"
+
+/**
+ * Base class for all the default sections, implementing common behavior for them
+ */
+class FAdvancedRenamerSectionBase : public IAdvancedRenamerSection
+{
+public:
+	/** Init the given section */
+	virtual void Init(TSharedRef<IAdvancedRenamer> InRenamer) override;
+
+	/** Return the section information*/
+	virtual FAdvancedRenamerExecuteSection GetSection() const override;
+
+protected:
+	/** Mark the Renamer dirty, this will make the Renamer re-execute the Rename logic */
+	virtual void MarkRenamerDirty() const override;
+
+protected:
+	/** Section information of this extension */
+	FAdvancedRenamerExecuteSection Section;
+	
+	/** Weak pointer to the Renamer */
+	TWeakPtr<IAdvancedRenamer> RenamerWeakPtr;
+};
diff --git a/Engine/Plugins/Experimental/AdvancedRenamer/Source/AdvancedRenamer/Private/AdvancedRenamerSections/IAdvancedRenamerSection.h b/Engine/Plugins/Experimental/AdvancedRenamer/Source/AdvancedRenamer/Private/AdvancedRenamerSections/IAdvancedRenamerSection.h
new file mode 100644
--- /dev/null
+++ b/Engine/Plugins/Experimental/AdvancedRenamer/Source/AdvancedRenamer/Private/AdvancedRenamerSections/IAdvancedRenamerSection.h
@@ -0,0 +1,52 @@
+// Copyright Epic Games, Inc. All Rights Reserved.
+
+#pragma once
+
+#include "AdvancedRenamerExecuteSection.h"
+#include "Templates/SharedPointerFwd.h"
+
+class IAdvancedRenamer;
+class SWidget;
+
+/**
+ * Interface for all the sections for the AdvancedRenamer
+ */
+class IAdvancedRenamerSection : public TSharedFromThis<IAdvancedRenamerSection>
+{
+public:
+	/**
+	 * Create a new AdvancedRenamer section given its constructor params
+	 * @tparam InAdvancedRenamerSectionFactory AdvancedRenamer section class
+	 * @tparam InArgsType constructor parameter type
+	 * @param InArgs Additional constructor parameter
+	 * @return The newly created AdvancedRenamer section
+	 */
+	template <
+		typename InAdvancedRenamerSectionFactory,
+		typename... InArgsType
+		UE_REQUIRES(TIsDerivedFrom<InAdvancedRenamerSectionFactory, IAdvancedRenamerSection>::Value)
+	>
+	static TSharedRef<InAdvancedRenamerSectionFactory> MakeInstance(InArgsType&&... InArgs)
+	{
+		return MakeShared<InAdvancedRenamerSectionFactory>(Forward<InArgsType>(InArgs)...);
+	}
+
+	/** virtual destructor */
+	virtual ~IAdvancedRenamerSection() = default;
+
+	/** Init the given section */
+	virtual void Init(TSharedRef<IAdvancedRenamer> InRenamer) = 0;
+
+	/** Reset all values of the section to the default ones */
+	virtual void ResetToDefault() = 0;
+
+	/** Return the section information*/
+	virtual FAdvancedRenamerExecuteSection GetSection() const = 0;
+
+	/** Return the widget for this section*/
+	virtual TSharedRef<SWidget> GetWidget() = 0;
+
+protected:
+	/** Mark the Renamer dirty, this will make the Renamer re-execute the Rename logic */
+	virtual void MarkRenamerDirty() const = 0;
+};
diff --git a/Engine/Plugins/Experimental/AdvancedRenamer/Source/AdvancedRenamer/Private/AdvancedRenamerStyle.cpp b/Engine/Plugins/Experimental/AdvancedRenamer/Source/AdvancedRenamer/Private/AdvancedRenamerStyle.cpp
--- a/Engine/Plugins/Experimental/AdvancedRenamer/Source/AdvancedRenamer/Private/AdvancedRenamerStyle.cpp
+++ b/Engine/Plugins/Experimental/AdvancedRenamer/Source/AdvancedRenamer/Private/AdvancedRenamerStyle.cpp
@@ -1,11 +1,8 @@
 // Copyright Epic Games, Inc. All Rights Reserved.
 
 #include "AdvancedRenamerStyle.h"
-#include "Brushes/SlateImageBrush.h"
 #include "Interfaces/IPluginManager.h"
-#include "Misc/Paths.h"
 #include "Styling/AppStyle.h"
-#include "Styling/ISlateStyle.h"
 #include "Styling/SlateStyle.h"
 #include "Styling/SlateStyleRegistry.h"
 #include "Styling/SlateTypes.h"
@@ -17,7 +14,7 @@ void FAdvancedRenamerStyle::Initialize()
 {
 	if (!StyleInstance.IsValid())
 	{
-		StyleInstance = Create();
+		InitStyle();
 		FSlateStyleRegistry::RegisterSlateStyle(*StyleInstance);
 	}
 }
@@ -38,37 +35,44 @@ FName FAdvancedRenamerStyle::GetStyleSetName()
 	return StyleSetName;
 }
 
-TSharedRef<FSlateStyleSet> FAdvancedRenamerStyle::Create()
+void FAdvancedRenamerStyle::InitStyle()
 {
-	TSharedRef<FSlateStyleSet> Style = MakeShared<FSlateStyleSet>("AdvancedRenamer");
+	if (StyleInstance.IsValid())
+	{
+		return;
+	}
+
+	StyleInstance = MakeShared<FSlateStyleSet>("AdvancedRenamer");
 
 	TSharedPtr<IPlugin> Plugin = IPluginManager::Get().FindPlugin(TEXT("AdvancedRenamer"));
 	check(Plugin.IsValid());
 
 	if (Plugin.IsValid())
 	{
-		Style->SetContentRoot(FPaths::Combine(Plugin->GetBaseDir(), TEXT("Resources")));
+		StyleInstance->SetContentRoot(FPaths::Combine(Plugin->GetBaseDir(), TEXT("Resources")));
 	}
 
-	Style->Set("AdvancedRenamer.Image.Radio.BlackBackground", new FSlateVectorImageBrush(
-		FPaths::Combine(Style->GetContentRootDir(), "/Icons/radio-background.svg"), 
-		FVector2D(16.f, 16.f), FStyleColors::Foldout.GetSpecifiedColor()
-	));
+	const FSplitterStyle SplitterStyle = FSplitterStyle()
+		.SetHandleNormalBrush(FSlateNoResource())
+		.SetHandleHighlightBrush(FSlateNoResource());
+	
+	StyleInstance->Set("AdvancedRenamer.Style.Splitter", SplitterStyle);
+
+	FSlateBrush* BackgroundBorderBrush = new FSlateColorBrush(FLinearColor::FromSRGBColor(FColor(36, 36, 36)));
 
-	FCheckBoxStyle BlackRadioButton = FAppStyle::GetWidgetStyle<FCheckBoxStyle>("RadioButton");
-	BlackRadioButton.SetBackgroundImage(*Style->GetBrush("AdvancedRenamer.Image.Radio.BlackBackground"));
+	StyleInstance->Set("AdvancedRenamer.Style.BackgroundBorder", BackgroundBorderBrush);
 
-	Style->Set("AdvancedRenamer.Style.BlackRadioButton", BlackRadioButton);
+	const FTableViewStyle ListViewStyle = FTableViewStyle()
+		.SetBackgroundBrush(*FAdvancedRenamerStyle::Get().GetBrush("AdvancedRenamer.Style.BackgroundBorder"));
 
-	const FButtonStyle DarkButton = FButtonStyle(FAppStyle::Get().GetWidgetStyle<FButtonStyle>("SimpleButton"))
-		.SetNormal(FSlateColorBrush(FStyleColors::Recessed))
-		.SetHovered(FSlateColorBrush(FStyleColors::Hover))
-		.SetPressed(FSlateColorBrush(FStyleColors::Header))
-		.SetDisabled(FSlateColorBrush(FStyleColors::Dropdown));
+	StyleInstance->Set("AdvancedRenamer.Style.ListView", ListViewStyle);
 
-	Style->Set("AdvancedRenamer.Style.DarkButton", DarkButton);
+	FHeaderRowStyle HeaderRowStyle = FAppStyle::Get().GetWidgetStyle<FHeaderRowStyle>("TableView.Header");
+	HeaderRowStyle.SetHorizontalSeparatorThickness(0);
+	HeaderRowStyle.SetHorizontalSeparatorBrush(FSlateNoResource());
+	HeaderRowStyle.SetBackgroundBrush(FSlateColorBrush(FLinearColor::FromSRGBColor(FColor(47, 47, 47))));
 
-	return Style;
+	StyleInstance->Set("AdvancedRenamer.Style.HeaderRow", HeaderRowStyle);
 }
 
 const ISlateStyle& FAdvancedRenamerStyle::Get()
diff --git a/Engine/Plugins/Experimental/AdvancedRenamer/Source/AdvancedRenamer/Private/AdvancedRenamerStyle.h b/Engine/Plugins/Experimental/AdvancedRenamer/Source/AdvancedRenamer/Private/AdvancedRenamerStyle.h
--- a/Engine/Plugins/Experimental/AdvancedRenamer/Source/AdvancedRenamer/Private/AdvancedRenamerStyle.h
+++ b/Engine/Plugins/Experimental/AdvancedRenamer/Source/AdvancedRenamer/Private/AdvancedRenamerStyle.h
@@ -22,7 +22,7 @@ public:
 
 private:
 
-	static TSharedRef<FSlateStyleSet> Create();
+	static void InitStyle();
 	static TSharedPtr<FSlateStyleSet> StyleInstance;
 
 };
diff --git a/Engine/Plugins/Experimental/AdvancedRenamer/Source/AdvancedRenamer/Private/IAdvancedRenamer.h b/Engine/Plugins/Experimental/AdvancedRenamer/Source/AdvancedRenamer/Private/IAdvancedRenamer.h
--- a/Engine/Plugins/Experimental/AdvancedRenamer/Source/AdvancedRenamer/Private/IAdvancedRenamer.h
+++ b/Engine/Plugins/Experimental/AdvancedRenamer/Source/AdvancedRenamer/Private/IAdvancedRenamer.h
@@ -2,105 +2,17 @@
 
 #pragma once
 
-#include "Containers/ContainersFwd.h"
+#include "AdvancedRenamerExecuteSection.h"
 #include "Containers/UnrealString.h"
-#include "HAL/Platform.h"
-#include "Misc/CString.h"
 #include "Providers/IAdvancedRenamerProvider.h"
 #include "Templates/SharedPointerFwd.h"
 
-enum class EAdvancedRenamerSeachAndReplaceType : uint8
-{
-	None,
-	PlainText,
-	RegularExpression
-};
-
-/**
- * Options for the rename engine.
- *
- * - RemovePrefixSeparator and RemovePrefixCharacterCount are mutually exclusive.
- * - RemoveSuffixSeparator and RemoveSuffixCharacterCount are mutually exclusive.
- * - bAddSuffixNumber implies bRemoveSuffixNumer.
- */
-struct FAdvancedRenamerOptions
-{
-	/** Name => OtherName */
-	FString BaseName = ""; 
-
-	/** Name => Other_Name */
-	FString AddPrefix = ""; 
-
-	/** Other_Name, _ => Name */
-	FString RemovePrefixSeparator = "";
-
-	/** OtherName, 5 => Name */
-	int32 RemovePrefixCharacterCount = 0; 
-
-	/** Name => Name_Other */
-	FString AddSuffix = ""; 
-
-	/** Name => Name1 */
-	bool bAddSuffixNumber = false; 
-
-	/** Name, 5 => Name5 */
-	int32 AddSuffixNumberStart = 1; 
-
-	/** Name, Moo, 1, 5 => Name1, Moo6 */
-	int32 AddSuffixNumberStep = 1;
-
-	/** Name_Other, _ => Name */
-	FString RemoveSuffixSeparator = ""; 
-
-	/** NameOther, 5 => Name */
-	int32 RemoveSuffixCharacterCount = 0; 
-
-	/** Name5 => Name */
-	bool bRemoveSuffixNumber = false; 
-
-	EAdvancedRenamerSeachAndReplaceType SearchAndReplaceType = EAdvancedRenamerSeachAndReplaceType::None;
-	ESearchCase::Type SearchAndReplaceCase = ESearchCase::CaseSensitive;
-
-	/** Name, am, ot => Note */
-	FString SearchAndReplaceFromText = "";
-
-	/** Name, am, ot => Note */
-	FString SearchAndReplaceToText = "";
-
-	bool HasBaseName() const
-	{
-		return !BaseName.IsEmpty();
-	}
-
-	bool HasPrefix() const
-	{
-		return !AddPrefix.IsEmpty()
-			|| !RemovePrefixSeparator.IsEmpty()
-			|| RemovePrefixCharacterCount > 0;
-	}
-
-	bool HasSuffix() const
-	{
-		return !AddSuffix.IsEmpty()
-			|| (bAddSuffixNumber && AddSuffixNumberStep > 0)
-			|| !RemoveSuffixSeparator.IsEmpty()
-			|| RemoveSuffixCharacterCount > 0
-			|| bRemoveSuffixNumber;			
-	}
-
-	bool HasSearchAndReplace() const
-	{
-		return SearchAndReplaceType != EAdvancedRenamerSeachAndReplaceType::None
-			&& !SearchAndReplaceFromText.IsEmpty();
-	}
-};
-
 struct FAdvancedRenamerPreview
 {
 	FAdvancedRenamerPreview(int32 InHash, const FString InOriginalName)
 		: Hash(InHash)
 		, OriginalName(InOriginalName)
-		, NewName(FString(""))
+		, NewName(TEXT(""))
 	{
 	}
 
@@ -120,18 +32,14 @@ public:
 
 	virtual const TSharedRef<IAdvancedRenamerProvider>& GetProvider() const = 0;
 
-	virtual const FAdvancedRenamerOptions& GetOptions() const = 0;
-
-	/** Make sure to run UpdatePreviews or MarkDirty if you make changes directly. */
-	virtual FAdvancedRenamerOptions& GetOptions() = 0;
-
-	virtual void SetOption(const FAdvancedRenamerOptions& InOptions) = 0;
-
 	virtual const TArray<TSharedPtr<FAdvancedRenamerPreview>>& GetPreviews() = 0;
 
 	/** Returns the preview for the item at the given index. */
 	virtual TSharedPtr<FAdvancedRenamerPreview> GetPreview(int32 InIndex) const = 0;
 
+	/** Add a section to this Renamer */
+	virtual void AddSection(FAdvancedRenamerExecuteSection InSection) = 0;
+
 	/** True if there are any items actually renamed by the preview generator. */
 	virtual bool HasRenames() const = 0;
 
@@ -141,8 +49,8 @@ public:
 	virtual void MarkClean() = 0;
 
 	/** Executes the rename on the given name. */
-	virtual FString ApplyRename(const FString& InName, int32 InIndex) const = 0;
-
+	virtual FString ApplyRename(const FString& InName) = 0;
+	
 	/** Returns true if any names actually changed. */
 	virtual bool UpdatePreviews() = 0;
 
diff --git a/Engine/Plugins/Experimental/AdvancedRenamer/Source/AdvancedRenamer/Private/Integrations/AdvancedRenamerContentBrowserIntegration.cpp b/Engine/Plugins/Experimental/AdvancedRenamer/Source/AdvancedRenamer/Private/Integrations/AdvancedRenamerContentBrowserIntegration.cpp
--- a/Engine/Plugins/Experimental/AdvancedRenamer/Source/AdvancedRenamer/Private/Integrations/AdvancedRenamerContentBrowserIntegration.cpp
+++ b/Engine/Plugins/Experimental/AdvancedRenamer/Source/AdvancedRenamer/Private/Integrations/AdvancedRenamerContentBrowserIntegration.cpp
@@ -29,8 +29,8 @@ namespace UE::AdvancedRenamer::Private
 	void ExtendAssetMenu(FMenuBuilder& MenuBuilder, const TArray<FAssetData> AssetArray)
 	{
 		MenuBuilder.AddMenuEntry(
-			LOCTEXT("AdvancedRename", "Advanced Rename"),
-			LOCTEXT("AdvancedRenameTooltip", "Opens the Advanced Renamer Panel to rename all selected assets."),
+			LOCTEXT("AdvancedRename", "Batch Rename"),
+			LOCTEXT("AdvancedRenameTooltip", "Opens the Batch Renamer Panel to rename all selected assets."),
 			FSlateIcon(FAppStyle::GetAppStyleSetName(), "ContentBrowser.AssetActions.Rename"),
 			FUIAction(FExecuteAction::CreateStatic(&OpenAdvancedRenamer, AssetArray))
 		);
diff --git a/Engine/Plugins/Experimental/AdvancedRenamer/Source/AdvancedRenamer/Private/Integrations/AdvancedRenamerLevelEditorIntegration.cpp b/Engine/Plugins/Experimental/AdvancedRenamer/Source/AdvancedRenamer/Private/Integrations/AdvancedRenamerLevelEditorIntegration.cpp
--- a/Engine/Plugins/Experimental/AdvancedRenamer/Source/AdvancedRenamer/Private/Integrations/AdvancedRenamerLevelEditorIntegration.cpp
+++ b/Engine/Plugins/Experimental/AdvancedRenamer/Source/AdvancedRenamer/Private/Integrations/AdvancedRenamerLevelEditorIntegration.cpp
@@ -102,20 +102,23 @@ namespace UE::AdvancedRenamer::Private
 	}
 }
 
+static const TArray<FName, TInlineAllocator<2>> Menus = {
+	"LevelEditor.ActorContextMenu.EditSubMenu",
+	"LevelEditor.LevelEditorSceneOutliner.ContextMenu.EditSubMenu"
+};
+
 void FAdvancedRenamerLevelEditorIntegration::Initialize()
 {
 	using namespace UE::AdvancedRenamer::Private;
 
 	FLevelEditorModule& LevelEditorModule = FModuleManager::LoadModuleChecked<FLevelEditorModule>("LevelEditor");
 	LevelEditorCreatedDelegateHandle = LevelEditorModule.OnLevelEditorCreated().AddStatic(&OnLevelEditorCreated);
+}
 
+void FAdvancedRenamerLevelEditorIntegration::InitializeMenu()
+{
 	const FAdvancedRenamerCommands& AdvRenCommands = FAdvancedRenamerCommands::Get();
 
-	static const TArray<FName, TInlineAllocator<2>> Menus = {
-		"LevelEditor.ActorContextMenu.EditSubMenu",
-		"LevelEditor.LevelEditorSceneOutliner.ContextMenu.EditSubMenu"
-	};
-
 	const TAttribute<FText> TextAttribute;
 	const FSlateIcon RenameIcon = FSlateIcon(FAppStyle::GetAppStyleSetName(), "GenericCommands.Rename");
 
@@ -147,5 +150,18 @@ void FAdvancedRenamerLevelEditorIntegration::Shutdown()
 	if (FLevelEditorModule* LevelEditorModule = FModuleManager::LoadModulePtr<FLevelEditorModule>("LevelEditor"))
 	{
 		LevelEditorModule->OnLevelEditorCreated().Remove(LevelEditorCreatedDelegateHandle);
+		LevelEditorCreatedDelegateHandle.Reset();
+	}
+}
+
+void FAdvancedRenamerLevelEditorIntegration::ShutdownMenu()
+{
+	if (UToolMenus* ToolMenus = UToolMenus::Get())
+	{
+		for (const FName& Menu : Menus)
+		{
+			ToolMenus->RemoveEntry(Menu, NAME_None, TEXT("RenameSelectedActors"));
+			ToolMenus->RemoveEntry(Menu, NAME_None, TEXT("RenameSharedClassActors"));
+		}
 	}
 }
diff --git a/Engine/Plugins/Experimental/AdvancedRenamer/Source/AdvancedRenamer/Private/Integrations/AdvancedRenamerLevelEditorIntegration.h b/Engine/Plugins/Experimental/AdvancedRenamer/Source/AdvancedRenamer/Private/Integrations/AdvancedRenamerLevelEditorIntegration.h
--- a/Engine/Plugins/Experimental/AdvancedRenamer/Source/AdvancedRenamer/Private/Integrations/AdvancedRenamerLevelEditorIntegration.h
+++ b/Engine/Plugins/Experimental/AdvancedRenamer/Source/AdvancedRenamer/Private/Integrations/AdvancedRenamerLevelEditorIntegration.h
@@ -6,5 +6,7 @@ class FAdvancedRenamerLevelEditorIntegration
 {
 public:
 	static void Initialize();
+	static void InitializeMenu();
 	static void Shutdown();	
+	static void ShutdownMenu();
 };
diff --git a/Engine/Plugins/Experimental/AdvancedRenamer/Source/AdvancedRenamer/Private/Slate/SAdvancedRenamerPanel.cpp b/Engine/Plugins/Experimental/AdvancedRenamer/Source/AdvancedRenamer/Private/Slate/SAdvancedRenamerPanel.cpp
--- a/Engine/Plugins/Experimental/AdvancedRenamer/Source/AdvancedRenamer/Private/Slate/SAdvancedRenamerPanel.cpp
+++ b/Engine/Plugins/Experimental/AdvancedRenamer/Source/AdvancedRenamer/Private/Slate/SAdvancedRenamerPanel.cpp
@@ -1,60 +1,31 @@
 // Copyright Epic Games, Inc. All Rights Reserved.
 
 #include "Slate/SAdvancedRenamerPanel.h"
+
+#include "AdvancedRenamerModule.h"
+#include "AdvancedRenamerSections/IAdvancedRenamerSection.h"
 #include "AdvancedRenamerStyle.h"
 #include "EngineAnalytics.h"
+#include "Framework/Application/SlateApplication.h"
 #include "Framework/Commands/GenericCommands.h"
 #include "Framework/MultiBox/MultiBoxBuilder.h"
+#include "GenericPlatform/GenericPlatformMath.h"
 #include "IAdvancedRenamer.h"
 #include "ScopedTransaction.h"
-#include "Styling/AppStyle.h"
-#include "Styling/StyleColors.h"
-#include "Widgets/Colors/SColorBlock.h"
+#include "Framework/Notifications/NotificationManager.h"
+#include "Utils/AdvancedRenamerSlateUtils.h"
 #include "Widgets/Input/SButton.h"
-#include "Widgets/Input/SCheckBox.h"
-#include "Widgets/Input/SEditableTextBox.h"
-#include "Widgets/Input/SMultiLineEditableTextBox.h"
-#include "Widgets/Input/SSpinBox.h"
 #include "Widgets/Layout/SBorder.h"
-#include "Widgets/Layout/SBox.h"
-#include "Widgets/Layout/SScrollBox.h"
-#include "Widgets/SCanvas.h"
-#include "Widgets/SNullWidget.h"
+#include "Widgets/Layout/SSpacer.h"
+#include "Widgets/Layout/SSplitter.h"
+#include "Widgets/SBoxPanel.h"
+#include "Widgets/SWindow.h"
+#include "Widgets/Notifications/SNotificationList.h"
+#include "Widgets/Text/STextBlock.h"
 #include "Widgets/Views/SListView.h"
-#include "Widgets/Views/STableRow.h"
 
 #define LOCTEXT_NAMESPACE "SAdvancedRenamerPanel"
 
-namespace UE::AdvancedRenamer::Private
-{
-	const FVector2D WindowSize = {600.f, 500.f};
-	const FSlateFontInfo TitleFont = FCoreStyle::GetDefaultFontStyle("Regular", 12);
-	const FSlateFontInfo RegularFont = FCoreStyle::GetDefaultFontStyle("Regular", 10);
-	constexpr float LeftBlockWidth = 304.f;
-	const FVector2D TitleSize = {150.f, 20.f};
-	constexpr float TitleOffsetY = 25.f;
-	const FVector2D CheckboxOffset = {0.f, 2.f};
-	const FVector2D CheckboxSize = {18.f, 18.f};
-	const FVector2D RadioOffset = {1.f, 2.f};
-	constexpr float LabelStart = 25.f;
-	constexpr float LabelOffsetY = 3.f;
-	const FVector2D LabelSize = {125.f, 18.f};
-	constexpr float EntryStart = 145.f;
-	const FVector2D EntrySize = {155.f, 21.f};
-	constexpr float LineHeight = 26.f;
-	const FVector2D SeparatorSize = {32.f, 21.f};
-	const FVector2D SpinSize1 = {32.f, 21.f};
-	const FVector2D SpinSize2 = {39.f, 21.f};
-	const FVector2D SpinSize3 = {46.f, 21.f};
-	constexpr float RightBlockOffsetX = 314.f;
-	const FVector2D RightBlockSize = {281.f, 490.f};
-	const FVector2D ListViewSize = {277.f, 431.f};
-	constexpr float ListLineHeight = 15.f;
-	constexpr float ApplyButtonHeight = 25.f;
-	static FName OriginalNameColumnName = "OriginalName";
-	static FName NewNameColumnName = "NewName";
-}
-
 class SAdvancedRenamerPreviewListRow : public SMultiColumnTableRow<TSharedPtr<FAdvancedRenamerPreview>>
 {
 public:
@@ -68,12 +39,12 @@ public:
 		ItemWeak = InRowItem;
 
 		SMultiColumnTableRow<TSharedPtr<FAdvancedRenamerPreview>>::Construct(FSuperRowType::FArguments(), InOwnerTableView);
-		SetBorderImage(TAttribute<const FSlateBrush*>(this, &SAdvancedRenamerPreviewListRow::GetBorder));
+		SetBorderImage(FAdvancedRenamerStyle::Get().GetBrush("AdvancedRenamer.Style.BackgroundBorder"));
 	}
 
 	virtual TSharedRef<SWidget> GenerateWidgetForColumn(const FName& InColumnName) override
 	{
-		using namespace UE::AdvancedRenamer::Private;
+		using namespace AdvancedRenamerSlateUtils::Default;
 
 		if (InColumnName != OriginalNameColumnName
 			&& InColumnName != NewNameColumnName)
@@ -96,22 +67,22 @@ public:
 		}
 
 		TSharedRef<STextBlock> TextBlock = SNew(STextBlock)
-			.ColorAndOpacity(FStyleColors::AccentWhite.GetSpecifiedColor())
-			.Font(RegularFont)
-			.Margin(FMargin(5.f, 0.f, 5.f, 0.f));
+			.Font(RegularFont);
+
+		TSharedRef<SBox> TextBlockBox = SNew(SBox)
+			.Padding(FMargin(8.f, 6.f, 0.f, 6.f));
 
 		if (InColumnName == OriginalNameColumnName)
 		{
 			TextBlock->SetText(FText::FromString(Item->OriginalName));
-			Panel->MinDesiredOriginalNameWidth = FMath::Max(Panel->MinDesiredOriginalNameWidth, TextBlock->ComputeDesiredSize(1.f).X);
 		}
 		else if (InColumnName == NewNameColumnName)
 		{
 			TextBlock->SetText(FText::FromString(Item->NewName));
-			Panel->MinDesiredNewNameWidth = FMath::Max(Panel->MinDesiredOriginalNameWidth, TextBlock->ComputeDesiredSize(1.f).X);
 		}
 
-		return TextBlock;
+		TextBlockBox->SetContent(TextBlock);
+		return TextBlockBox;
 	}
 
 protected:
@@ -123,11 +94,6 @@ void SAdvancedRenamerPanel::Construct(const FArguments& InArgs, const TSharedRef
 {
 	Renamer = InRenamer;
 
-	bRemovePrefixSeparator = false;
-	bRemovePrefixNumChars = false;
-	bRemoveSuffixSeparator = false;
-	bRemoveSuffixNumChars = false;
-
 	CommandList = MakeShared<FUICommandList>();
 	CommandList->MapAction(
 		FGenericCommands::Get().Delete,
@@ -135,637 +101,251 @@ void SAdvancedRenamerPanel::Construct(const FArguments& InArgs, const TSharedRef
 		FCanExecuteAction()
 	);
 
-	using namespace UE::AdvancedRenamer::Private;
-
-	TSharedRef<SCanvas> Canvas = SNew(SCanvas)
-		+ SCanvas::Slot()
-		.Position(FVector2D(0.f, 0.f))
-		.Size(WindowSize)
-		[
-			SNew(SColorBlock)
-			.Color(FStyleColors::Background.GetSpecifiedColor())
-		];
-
-	CreateLeftPane(Canvas);
-	CreateRightPane(Canvas);
+	using namespace AdvancedRenamerSlateUtils::Default;
 
 	ChildSlot
 	[
-		Canvas
-	];
-
-	if (FEngineAnalytics::IsAvailable())
-	{
-		FEngineAnalytics::GetProvider().RecordEvent(TEXT("Editor.Usage.AdvancedRenamer.Opened"));
-	}
-}
-
-void SAdvancedRenamerPanel::CreateLeftPane(const TSharedRef<SCanvas>& InCanvas)
-{	
-	using namespace UE::AdvancedRenamer::Private;
-
-	InCanvas->AddSlot()
-		.Position(FVector2D(5.f, 5.f))
-		.Size(FVector2D(LeftBlockWidth, 50.f))
-		[
-			CreateBaseName()
-		];
-
-	InCanvas->AddSlot()
-		.Position(FVector2D(5.f, 60.f))
-		.Size(FVector2D(LeftBlockWidth, 102.f))
-		[
-			CreatePrefix()
-		];
-
-	InCanvas->AddSlot()
-		.Position(FVector2D(5.f, 167.f))
-		.Size(FVector2D(LeftBlockWidth, 154.f))
-		[
-			CreateSuffix()
-		];
+		SNew(SVerticalBox)
 
-	InCanvas->AddSlot()
-		.Position(FVector2D(5.f, 327.f))
-		.Size(FVector2D(LeftBlockWidth, 168.f))
+		+ SVerticalBox::Slot()
+		.FillHeight(1.f)
 		[
-			CreateSearchAndReplace()
-		];
-}
-
-TSharedRef<SWidget> SAdvancedRenamerPanel::CreateBaseName()
-{
-	using namespace UE::AdvancedRenamer::Private;
-
-	return SNew(SBorder)
-		.BorderImage(FAppStyle::Get().GetBrush("DetailsView.CategoryTop"))
-		.Content()
-		[
-			SNew(SCanvas)
-			+ SCanvas::Slot()
-			.Position(FVector2D(0.f, LabelOffsetY))
-			.Size(TitleSize)
-			[
-				SNew(STextBlock)
-				.ColorAndOpacity(FStyleColors::AccentBlue.GetSpecifiedColor())
-				.Font(TitleFont)
-				.Text(LOCTEXT("BaseNameTitle", "Name"))
-			]
-			+ SCanvas::Slot()
-			.Position(FVector2D(LabelStart * 2.f, TitleOffsetY + LabelOffsetY))
-			.Size(LabelSize)
-			[
-				SNew(STextBlock)
-				.ColorAndOpacity(FStyleColors::AccentWhite.GetSpecifiedColor())
-				.Font(RegularFont)
-				.Text(LOCTEXT("SetBaseName", "Set Base Name"))
-			]
-			+ SCanvas::Slot()
-			.Position(FVector2D(EntryStart, TitleOffsetY))
-			.Size(EntrySize)
+			SNew(SBorder)
+			.BorderImage(FAppStyle::GetBrush("NoBorder"))
+			.Padding(2.f, 2.f)
 			[
-				SAssignNew(BaseNameTextBox, SEditableTextBox)
-				.BackgroundColor(FStyleColors::Background.GetSpecifiedColor())
-				.ForegroundColor(FStyleColors::AccentWhite.GetSpecifiedColor())
-				.Font(RegularFont)
-				.HintText(LOCTEXT("BaseNameHint", "Base name"))
-				.OnTextChanged(this, &SAdvancedRenamerPanel::OnBaseNameChanged)
+				SNew(SBox)
+				[
+					SNew(SSplitter)
+					.PhysicalSplitterHandleSize(2.f)
+					.MinimumSlotHeight(499.5f)
+					.Style(&FAdvancedRenamerStyle::Get().GetWidgetStyle<FSplitterStyle>("AdvancedRenamer.Style.Splitter"))
+
+					+ SSplitter::Slot()
+					.Value(0.505f)
+					.MinSize(365.f)
+					[
+						SAssignNew(LeftSideVerticalBox, SVerticalBox)
+					]
+
+					+ SSplitter::Slot()
+					.Value(0.495f)
+					[
+						SAssignNew(RightSideBox, SBox)
+					]
+				]
 			]
-		];
-}
+		]
 
-TSharedRef<SWidget> SAdvancedRenamerPanel::CreatePrefix()
-{
-	using namespace UE::AdvancedRenamer::Private;
-
-	return SNew(SBorder)
-		.BorderImage(FAppStyle::Get().GetBrush("DetailsView.CategoryTop"))
-		.Content()
+		+ SVerticalBox::Slot()
+		.AutoHeight()
 		[
-			SNew(SCanvas)
-
-			// Title
-			+ SCanvas::Slot()
-			.Position(FVector2D(0.f, LabelOffsetY))
-			.Size(TitleSize)
-			[
-				SNew(STextBlock)
-				.ColorAndOpacity(FStyleColors::AccentBlue.GetSpecifiedColor())
-				.Font(TitleFont)
-				.Text(LOCTEXT("PrefixTitle", "Prefix"))
-			]
-			//
-			+ SCanvas::Slot()
-			.Position(FVector2D(LabelStart * 2.f, TitleOffsetY + LabelOffsetY))
-			.Size(LabelSize)
-			[
-				SNew(STextBlock)
-				.ColorAndOpacity(FStyleColors::AccentWhite.GetSpecifiedColor())
-				.Font(RegularFont)
-				.Text(LOCTEXT("AddPrefix", "Add Prefix"))
-			]
-			+ SCanvas::Slot()
-			.Position(FVector2D(EntryStart, TitleOffsetY))
-			.Size(EntrySize)
+			SNew(SBorder)
+			.BorderImage(FAdvancedRenamerStyle::Get().GetBrush("AdvancedRenamer.Style.BackgroundBorder"))
 			[
-				SAssignNew(PrefixTextBox, SEditableTextBox)
-				.BackgroundColor(FStyleColors::Background.GetSpecifiedColor())
-				.ForegroundColor(FStyleColors::AccentWhite.GetSpecifiedColor())
-				.Font(RegularFont)
-				.HintText(LOCTEXT("PrefixHint", "Prefix"))
-				.OnTextChanged(this, &SAdvancedRenamerPanel::OnPrefixChanged)
+				SAssignNew(ApplyResetCancelBox, SBox)
 			]
+		]
+	];
 
-			// Remove prefix
-			+ SCanvas::Slot()
-			.Position(FVector2D(LabelStart, TitleOffsetY + LineHeight) + RadioOffset)
-			.Size(CheckboxSize)
-			[
-				SAssignNew(PrefixRemoveCheckBox, SCheckBox)
-				.Style(&FAdvancedRenamerStyle::Get().GetWidgetStyle<FCheckBoxStyle>("AdvancedRenamer.Style.BlackRadioButton"))
-				.IsChecked(this, &SAdvancedRenamerPanel::IsPrefixRemoveChecked)
-				.OnCheckStateChanged(this, &SAdvancedRenamerPanel::OnPrefixRemoveCheckBoxChanged)
-			]
-			+ SCanvas::Slot()
-			.Position(FVector2D(LabelStart * 2.f, TitleOffsetY + LabelOffsetY + LineHeight))
-			.Size(LabelSize)
-			[
-				SNew(STextBlock)
-				.ColorAndOpacity(FStyleColors::AccentWhite.GetSpecifiedColor())
-				.Font(RegularFont)
-				.Text(LOCTEXT("PrefixRemove", "Remove Prefix"))
-			]
-			+ SCanvas::Slot()
-			.Position(FVector2D(EntryStart, TitleOffsetY + LineHeight))
-			.Size(SeparatorSize)
-			[
-				SAssignNew(PrefixSeparatorTextBox, SEditableTextBox)
-				.BackgroundColor(FStyleColors::Background.GetSpecifiedColor())
-				.ForegroundColor(FStyleColors::AccentWhite.GetSpecifiedColor())
-				.Font(RegularFont)
-				.IsEnabled(this, &SAdvancedRenamerPanel::IsPrefixRemoveSeparatorEnabled)
-				.Text(LOCTEXT("Underscore", "_"))
-				.OnVerifyTextChanged(this, &SAdvancedRenamerPanel::OnPrefixSeparatorVerifyTextChanged)
-				.OnTextChanged(this, &SAdvancedRenamerPanel::OnPrefixSeparatorChanged)
-			]
-			+ SCanvas::Slot()
-			.Position(FVector2D(EntryStart + SeparatorSize.X + 5.f, TitleOffsetY + LabelOffsetY + LineHeight))
-			.Size(LabelSize)
-			[
-				SNew(STextBlock)
-				.ColorAndOpacity(FStyleColors::AccentWhite.GetSpecifiedColor())
-				.Font(RegularFont)
-				.Text(LOCTEXT("Separator", "(separator)"))
-			]
+	TArray<TSharedPtr<IAdvancedRenamerSection>> RegisteredSections = IAdvancedRenamerModule::Get().GetRegisteredSections();
+	for (int32 Index = 0; Index < RegisteredSections.Num(); Index++)
+	{
+		RegisteredSections[Index]->Init(InRenamer);
+		FMargin Padding = FMargin(0.f, 1.f);
 
-			// Remove first
-			+ SCanvas::Slot()
-			.Position(FVector2D(LabelStart, TitleOffsetY + LineHeight * 2.f) + RadioOffset)
-			.Size(CheckboxSize)
-			[
-				SAssignNew(PrefixRemoveCharactersCheckBox, SCheckBox)
-				.Style(&FAdvancedRenamerStyle::Get().GetWidgetStyle<FCheckBoxStyle>("AdvancedRenamer.Style.BlackRadioButton"))
-				.IsChecked(this, &SAdvancedRenamerPanel::IsPrefixRemoveCharactersChecked)
-				.OnCheckStateChanged(this, &SAdvancedRenamerPanel::OnPrefixRemoveCharactersCheckBoxChanged)
-			]
-			+ SCanvas::Slot()
-			.Position(FVector2D(LabelStart * 2.f, TitleOffsetY + LabelOffsetY + LineHeight * 2.f))
-			.Size(LabelSize)
-			[
-				SNew(STextBlock)
-				.ColorAndOpacity(FStyleColors::AccentWhite.GetSpecifiedColor())
-				.Font(RegularFont)
-				.Text(LOCTEXT("PrefixRemoveChars", "Remove First"))
-			]
-			+ SCanvas::Slot()
-			.Position(FVector2D(EntryStart, TitleOffsetY + LineHeight * 2.f))
-			.Size(SpinSize1)
-			[
+		if (Index == 0)
+		{
+			// The first Section Search and Replace will fill the height by design
+			Padding = FMargin(0.f, 0.f, 0.f, 1.f);
+			LeftSideVerticalBox->AddSlot()
+				.Padding(Padding)
+				.FillHeight(1.f)
+				[
+					RegisteredSections[Index]->GetWidget()
+				];
+			continue;
+		}
 
-				SAssignNew(PrefixRemoveCharactersSpinBox, SSpinBox<uint8>)
-				.Style(&FAppStyle::Get(), "Menu.SpinBox")
-				.Font(RegularFont)
-				.MinValue(1)
-				.MaxValue(9)
-				.Value(1)
-				.IsEnabled(this, &SAdvancedRenamerPanel::IsPrefixRemoveNumCharsEnabled)
-				.OnValueChanged(this, &SAdvancedRenamerPanel::OnPrefixRemoveCharactersChanged)
-			]
-			+ SCanvas::Slot()
-			.Position(FVector2D(EntryStart + SeparatorSize.X + 5.f, TitleOffsetY + LabelOffsetY + LineHeight * 2.f))
-			.Size(LabelSize)
-			[
-				SNew(STextBlock)
-				.ColorAndOpacity(FStyleColors::AccentWhite.GetSpecifiedColor())
-				.Font(RegularFont)
-				.Text(LOCTEXT("Characters", "character(s)"))
-			]
-		];
-}
+		if (Index == RegisteredSections.Num() - 1)
+		{
+			Padding = FMargin(0.f, 1.f, 0.f, 0.f);
+		}
 
-TSharedRef<SWidget> SAdvancedRenamerPanel::CreateSuffix()
-{
-	using namespace UE::AdvancedRenamer::Private;
+		LeftSideVerticalBox->AddSlot()
+			.Padding(Padding)
+			.AutoHeight()
+			[
+				RegisteredSections[Index]->GetWidget()
+			];
+	}
 
-	return SNew(SBorder)
-		.BorderImage(FAppStyle::Get().GetBrush("DetailsView.CategoryTop"))
-		.Content()
-		[
-			SNew(SCanvas)
+	CreateRightPanel();
+	CreateApplyResetCancelPanel();
 
-			// Title
-			+ SCanvas::Slot()
-			.Position(FVector2D(0.f, LabelOffsetY))
-			.Size(TitleSize)
-			[
-				SNew(STextBlock)
-				.ColorAndOpacity(FStyleColors::AccentBlue.GetSpecifiedColor())
-				.Font(TitleFont)
-				.Text(LOCTEXT("SuffixTitle", "Suffix"))
-			]
+	Renamer->MarkDirty();
 
-			// Suffix name
-			+ SCanvas::Slot()
-			.Position(FVector2D(LabelStart * 2.f, TitleOffsetY + LabelOffsetY))
-			.Size(LabelSize)
-			[
-				SNew(STextBlock)
-				.ColorAndOpacity(FStyleColors::AccentWhite.GetSpecifiedColor())
-				.Font(RegularFont)
-				.Text(LOCTEXT("AddSuffix", "Add Suffix"))
-			]
-			+ SCanvas::Slot()
-			.Position(FVector2D(EntryStart, TitleOffsetY))
-			.Size(EntrySize)
-			[
-				SAssignNew(SuffixTextBox, SEditableTextBox)
-				.BackgroundColor(FStyleColors::Background.GetSpecifiedColor())
-				.ForegroundColor(FStyleColors::AccentWhite.GetSpecifiedColor())
-				.Font(RegularFont)
-				.HintText(LOCTEXT("SuffixHint", "Suffix"))
-				.OnTextChanged(this, &SAdvancedRenamerPanel::OnSuffixChanged)
-			]
+	if (FEngineAnalytics::IsAvailable())
+	{
+		FEngineAnalytics::GetProvider().RecordEvent(TEXT("Editor.Usage.AdvancedRenamer.Opened"));
+	}
+}
 
-			// Remove Suffix
-			+ SCanvas::Slot()
-			.Position(FVector2D(LabelStart, TitleOffsetY + LineHeight) + RadioOffset)
-			.Size(CheckboxSize)
-			[
-				SAssignNew(SuffixRemoveCheckBox, SCheckBox)
-				.Style(&FAdvancedRenamerStyle::Get().GetWidgetStyle<FCheckBoxStyle>("AdvancedRenamer.Style.BlackRadioButton"))
-				.IsChecked(this, &SAdvancedRenamerPanel::IsSuffixRemoveChecked)
-				.OnCheckStateChanged(this, &SAdvancedRenamerPanel::OnSuffixRemoveCheckBoxChanged)
-			]
-			+ SCanvas::Slot()
-			.Position(FVector2D(LabelStart * 2.f, TitleOffsetY + LabelOffsetY + LineHeight))
-			.Size(LabelSize)
-			[
-				SNew(STextBlock)
-				.ColorAndOpacity(FStyleColors::AccentWhite.GetSpecifiedColor())
-				.Font(RegularFont)
-				.Text(LOCTEXT("SuffixRemove", "Remove Suffix"))
-			]
-			+ SCanvas::Slot()
-			.Position(FVector2D(EntryStart, TitleOffsetY + LineHeight))
-			.Size(SeparatorSize)
-			[
-				SAssignNew(SuffixSeparatorTextBox, SEditableTextBox)
-				.BackgroundColor(FStyleColors::Background.GetSpecifiedColor())
-				.ForegroundColor(FStyleColors::AccentWhite.GetSpecifiedColor())
-				.Font(RegularFont)
-				.IsEnabled(this, &SAdvancedRenamerPanel::IsSuffixRemoveSeparatorEnabled)
-				.Text(LOCTEXT("Underscore", "_"))
-				.OnVerifyTextChanged(this, &SAdvancedRenamerPanel::OnSuffixSeparatorVerifyTextChanged)
-				.OnTextChanged(this, &SAdvancedRenamerPanel::OnSuffixSeparatorChanged)
-			]
-			+ SCanvas::Slot()
-			.Position(FVector2D(EntryStart + SeparatorSize.X + 5.f, TitleOffsetY + LabelOffsetY + LineHeight))
-			.Size(LabelSize)
-			[
-				SNew(STextBlock)
-				.ColorAndOpacity(FStyleColors::AccentWhite.GetSpecifiedColor())
-				.Font(RegularFont)
-				.Text(LOCTEXT("Separator", "(separator)"))
-			]
+void SAdvancedRenamerPanel::CreateApplyResetCancelPanel()
+{
+	using namespace AdvancedRenamerSlateUtils::Default;
 
-			// Remove first
-			+ SCanvas::Slot()
-			.Position(FVector2D(LabelStart, TitleOffsetY + LineHeight * 2.f) + RadioOffset)
-			.Size(CheckboxSize)
-			[
-				SAssignNew(SuffixRemoveCharactersCheckBox, SCheckBox)
-				.Style(&FAdvancedRenamerStyle::Get().GetWidgetStyle<FCheckBoxStyle>("AdvancedRenamer.Style.BlackRadioButton"))
-				.IsChecked(this, &SAdvancedRenamerPanel::IsSuffixRemoveCharactersChecked)
-				.OnCheckStateChanged(this, &SAdvancedRenamerPanel::OnSuffixRemoveCharactersCheckBoxChanged)
-			]
-			+ SCanvas::Slot()
-			.Position(FVector2D(LabelStart * 2.f, TitleOffsetY + LabelOffsetY + LineHeight * 2.f))
-			.Size(LabelSize)
+	ApplyResetCancelBox->SetContent(SNew(SHorizontalBox)
+			+ SHorizontalBox::Slot()
+			.HAlign(HAlign_Right)
 			[
-				SNew(STextBlock)
-				.ColorAndOpacity(FStyleColors::AccentWhite.GetSpecifiedColor())
-				.Font(RegularFont)
-				.Text(LOCTEXT("SuffixRemoveChars", "Remove Last"))
-			]
-			+ SCanvas::Slot()
-			.Position(FVector2D(EntryStart, TitleOffsetY + LineHeight * 2.f))
-			.Size(SpinSize1)
-			[
-				SAssignNew(SuffixRemoveCharactersSpinBox, SSpinBox<uint8>)
-				.Style(&FAppStyle::Get(), "Menu.SpinBox")
-				.Font(RegularFont)
-				.MinValue(1)
-				.MaxValue(9)
-				.Value(1)
-				.IsEnabled(this, &SAdvancedRenamerPanel::IsSuffixRemoveNumCharsEnabled)
-				.OnValueChanged(this, &SAdvancedRenamerPanel::OnSuffixRemoveCharactersChanged)
-			]
-			+ SCanvas::Slot()
-			.Position(FVector2D(EntryStart + SeparatorSize.X + 5.f, TitleOffsetY + LabelOffsetY + LineHeight * 2.f))
-			.Size(LabelSize)
-			[
-				SNew(STextBlock)
-				.ColorAndOpacity(FStyleColors::AccentWhite.GetSpecifiedColor())
-				.Font(RegularFont)
-				.Text(LOCTEXT("Characters", "character(s)"))
+				SNew(SSpacer)
 			]
 
-			// Remove Number
-			+ SCanvas::Slot()
-			.Position(FVector2D(LabelStart, TitleOffsetY + LineHeight * 3.f) + CheckboxOffset)
-			.Size(CheckboxSize)
-			[
-				SAssignNew(SuffixRemoveNumberCheckBox, SCheckBox)
-				.IsChecked(this, &SAdvancedRenamerPanel::IsSuffixRemoveNumberChecked)
-				.OnCheckStateChanged(this, &SAdvancedRenamerPanel::OnSuffixRemoveNumberCheckBoxChanged)
-				.IsEnabled(this, &SAdvancedRenamerPanel::IsSuffixRemoveNumberCheckBoxEnabled)
-			]
-			+ SCanvas::Slot()
-			.Position(FVector2D(LabelStart * 2.f, TitleOffsetY + LabelOffsetY + LineHeight * 3.f))
-			.Size(LabelSize)
+			+ SHorizontalBox::Slot()
+			.Padding(ApplyButtonPadding)
+			.HAlign(HAlign_Right)
+			.AutoWidth()
 			[
-				SNew(STextBlock)
-				.ColorAndOpacity(FStyleColors::AccentWhite.GetSpecifiedColor())
-				.Font(RegularFont)
-				.Text(LOCTEXT("SuffixRemoveNumber", "Remove Number"))
+				CreateApplyButton()
 			]
 
-			// Number
-			+ SCanvas::Slot()
-			.Position(FVector2D(LabelStart, TitleOffsetY + LineHeight * 4.f) + CheckboxOffset)
-			.Size(CheckboxSize)
-			[
-				SAssignNew(SuffixNumberCheckBox, SCheckBox)
-				.IsChecked(this, &SAdvancedRenamerPanel::IsSuffixNumberChecked)
-				.OnCheckStateChanged(this, &SAdvancedRenamerPanel::OnSuffixNumberCheckBoxChanged)
-			]
-			+ SCanvas::Slot()
-			.Position(FVector2D(LabelStart * 2.f, TitleOffsetY + LabelOffsetY + LineHeight * 4.f))
-			.Size(LabelSize)
-			[
-				SNew(STextBlock)
-				.ColorAndOpacity(FStyleColors::AccentWhite.GetSpecifiedColor())
-				.Font(RegularFont)
-				.Text(LOCTEXT("SuffixNumber", "Add Number"))
-			]
-			+ SCanvas::Slot()
-			.Position(FVector2D(EntryStart , TitleOffsetY + LabelOffsetY + LineHeight * 4.f))
-			.Size(LabelSize)
+			+ SHorizontalBox::Slot()
+			.Padding(ResetButtonPadding)
+			.HAlign(HAlign_Right)
+			.AutoWidth()
 			[
-				SNew(STextBlock)
-				.ColorAndOpacity(FStyleColors::AccentWhite.GetSpecifiedColor())
-				.Font(RegularFont)
-				.Text(LOCTEXT("Start", "Start"))
+				CreateResetButton()
 			]
-			+ SCanvas::Slot()
-			.Position(FVector2D(EntryStart + 33.f, TitleOffsetY + LineHeight * 4.f))
-			.Size(SpinSize3)
-			[
 
-				SAssignNew(SuffixNumberStartSpinBox, SSpinBox<int32>)
-				.Style(&FAppStyle::Get(), "Menu.SpinBox")
-				.Font(RegularFont)
-				.MinValue(0)
-				.MaxValue(999)
-				.Value(1)
-				.IsEnabled(false)
-				.OnValueChanged(this, &SAdvancedRenamerPanel::OnSuffixNumberStartChanged)
-			]
-			+ SCanvas::Slot()
-			.Position(FVector2D(EntryStart + 33.f + SpinSize3.X + 5.f, TitleOffsetY + LabelOffsetY + LineHeight * 4.f))
-			.Size(LabelSize)
-			[
-				SNew(STextBlock)
-				.ColorAndOpacity(FStyleColors::AccentWhite.GetSpecifiedColor())
-				.Font(RegularFont)
-				.Text(LOCTEXT("Step", "Step"))
-			]
-			+ SCanvas::Slot()
-			.Position(FVector2D(EntryStart + 33.f + SpinSize3.X + 36.f, TitleOffsetY + LineHeight * 4.f))
-			.Size(SpinSize2)
+			+ SHorizontalBox::Slot()
+			.Padding(CancelButtonPadding)
+			.HAlign(HAlign_Right)
+			.AutoWidth()
 			[
+				CreateCancelButton()
+			]);
+}
 
-				SAssignNew(SuffixNumberStepSpinBox, SSpinBox<int32>)
-				.Style(&FAppStyle::Get(), "Menu.SpinBox")
-				.Font(RegularFont)
-				.MinValue(1)
-				.MaxValue(99)
-				.Value(1)
-				.IsEnabled(false)
-				.OnValueChanged(this, &SAdvancedRenamerPanel::OnSuffixNumberStepChanged)
+TSharedRef<SWidget> SAdvancedRenamerPanel::CreateApplyButton()
+{
+	return SNew(SBox)
+		.HeightOverride(25)
+		.WidthOverride(75)
+		[
+			SNew(SButton)
+			.IsEnabled(this, &SAdvancedRenamerPanel::IsApplyButtonEnabled)
+			.ButtonStyle(&FAppStyle::Get().GetWidgetStyle<FButtonStyle>("PrimaryButton"))
+			.ContentPadding(FMargin(-2.f, -1.f))
+			.OnClicked(this, &SAdvancedRenamerPanel::OnApplyButtonClicked)
+			.Content()
+			[
+				SNew(SBox)
+				.VAlign(VAlign_Center)
+				[
+					SNew(STextBlock)
+					.Justification(ETextJustify::Center)
+					.Text(LOCTEXT("AR_Apply", "Apply"))
+				]
 			]
 		];
 }
 
-TSharedRef<SWidget> SAdvancedRenamerPanel::CreateSearchAndReplace()
+TSharedRef<SWidget> SAdvancedRenamerPanel::CreateResetButton()
 {
-	using namespace UE::AdvancedRenamer::Private;
-
-	return SNew(SBorder)
-		.BorderImage(FAppStyle::Get().GetBrush("DetailsView.CategoryTop"))
-		.Content()
+	return SNew(SBox)
+		.HeightOverride(25)
+		.WidthOverride(75)
 		[
-			SNew(SCanvas)
-
-			// Title
-			+ SCanvas::Slot()
-			.Position(FVector2D(0.f, LabelOffsetY))
-			.Size(TitleSize)
-			[
-				SNew(STextBlock)
-				.ColorAndOpacity(FStyleColors::AccentBlue.GetSpecifiedColor())
-				.Font(TitleFont)
-				.Text(LOCTEXT("SearchReplaceTitle", "Search and Replace"))
-			]
-
-			+ SCanvas::Slot()
-			.Position(FVector2D(0.f, TitleOffsetY) + RadioOffset)
-			.Size(CheckboxSize)
-			[
-				SAssignNew(SearchReplacePlainTextCheckbox, SCheckBox)
-				.Style(&FAdvancedRenamerStyle::Get().GetWidgetStyle<FCheckBoxStyle>("AdvancedRenamer.Style.BlackRadioButton"))
-				.IsChecked(this, &SAdvancedRenamerPanel::IsSearchReplacePlainTextChecked)
-				.OnCheckStateChanged(this, &SAdvancedRenamerPanel::OnSearchReplacePlainTextCheckBoxChanged)
-			]
-
-			+ SCanvas::Slot()
-			.Position(FVector2D(0.f + LabelStart, TitleOffsetY + LabelOffsetY))
-			.Size(LabelSize)
-			[
-				SNew(STextBlock)
-				.ColorAndOpacity(FStyleColors::AccentWhite.GetSpecifiedColor())
-				.Font(RegularFont)
-				.Text(LOCTEXT("PlainText", "Plain Text"))
-			]
-
-			+ SCanvas::Slot()
-			.Position(FVector2D(92.f, TitleOffsetY) + RadioOffset)
-			.Size(CheckboxSize)
-			[
-				SAssignNew(SearchReplaceRegexCheckbox, SCheckBox)
-				.Style(&FAdvancedRenamerStyle::Get().GetWidgetStyle<FCheckBoxStyle>("AdvancedRenamer.Style.BlackRadioButton"))
-				.IsChecked(this, &SAdvancedRenamerPanel::IsSearchReplaceRegexChecked)
-				.OnCheckStateChanged(this, &SAdvancedRenamerPanel::OnSearchReplaceRegexCheckBoxChanged)
-			]
-
-			+ SCanvas::Slot()
-			.Position(FVector2D(92.f + LabelStart, TitleOffsetY + LabelOffsetY))
-			.Size(LabelSize)
-			[
-				SNew(STextBlock)
-				.ColorAndOpacity(FStyleColors::AccentWhite.GetSpecifiedColor())
-				.Font(RegularFont)
-				.Text(LOCTEXT("Regex", "Regex"))
-			]
-
-			+ SCanvas::Slot()
-			.Position(FVector2D(200.f, TitleOffsetY) + CheckboxOffset)
-			.Size(CheckboxSize)
-			[
-				SAssignNew(SearchReplaceIgnoreCaseCheckBox, SCheckBox)
-				.IsChecked(this, &SAdvancedRenamerPanel::IsSearchReplaceIgnoreCaseChecked)
-				.OnCheckStateChanged(this, &SAdvancedRenamerPanel::OnSearchReplaceIgnoreCaseCheckBoxChanged)
-			]
-
-			+ SCanvas::Slot()
-			.Position(FVector2D(200.f + LabelStart, TitleOffsetY + LabelOffsetY))
-			.Size(LabelSize)
+			SNew(SButton)
+			.ContentPadding(FMargin(-2.f, -1.f))
+			.OnClicked(this, &SAdvancedRenamerPanel::OnResetButtonClicked)
+			.Content()
 			[
-				SNew(STextBlock)
-				.ColorAndOpacity(FStyleColors::AccentWhite.GetSpecifiedColor())
-				.Font(RegularFont)
-				.Text(LOCTEXT("IgnoreCase", "Ignore Case"))
-			]
-
-			+ SCanvas::Slot()
-			.Position(FVector2D(0.f, TitleOffsetY + LineHeight))
-			.Size(FVector2D(LeftBlockWidth - 4.f, 54.f))
-			[
-				SAssignNew(SearchReplaceSearchTextBox, SMultiLineEditableTextBox)
-				.BackgroundColor(FStyleColors::Background.GetSpecifiedColor())
-				.ForegroundColor(FStyleColors::AccentWhite.GetSpecifiedColor())
-				.Font(RegularFont)
-				.AllowMultiLine(false)
-				.AutoWrapText(true)
-				.HintText(LOCTEXT("RegeSearchHint", "Search"))
-				.WrappingPolicy(ETextWrappingPolicy::AllowPerCharacterWrapping)
-				.OnTextChanged(this, &SAdvancedRenamerPanel::OnSearchReplaceSearchTextChanged)
+				SNew(SBox)
+				.VAlign(VAlign_Center)
+				[
+					SNew(STextBlock)
+					.Justification(ETextJustify::Center)
+					.Text(LOCTEXT("AR_Reset", "Reset"))
+				]
 			]
+		];
+}
 
-			+ SCanvas::Slot()
-			.Position(FVector2D(0.f, TitleOffsetY + LineHeight + 59.f))
-			.Size(FVector2D(LeftBlockWidth - 4.f, 54.f))
+TSharedRef<SWidget> SAdvancedRenamerPanel::CreateCancelButton()
+{
+	return SNew(SBox)
+		.HeightOverride(25)
+		.WidthOverride(75)
+		[
+			SNew(SButton)
+			.ContentPadding(FMargin(-2.f, -1.f))
+			.OnClicked(this, &SAdvancedRenamerPanel::OnCancelButtonClicked)
+			.Content()
 			[
-				SAssignNew(SearchReplaceReplaceTextBox, SMultiLineEditableTextBox)
-				.BackgroundColor(FStyleColors::Background.GetSpecifiedColor())
-				.ForegroundColor(FStyleColors::AccentWhite.GetSpecifiedColor())
-				.Font(RegularFont)
-				.AllowMultiLine(false)
-				.AutoWrapText(true)
-				.HintText(LOCTEXT("RegeReplaceHint", "Replace"))
-				.WrappingPolicy(ETextWrappingPolicy::AllowPerCharacterWrapping)
-				.OnTextChanged(this, &SAdvancedRenamerPanel::OnSearchReplaceReplaceTextChanged)
+				SNew(SBox)
+				.VAlign(VAlign_Center)
+				[
+					SNew(STextBlock)
+					.Justification(ETextJustify::Center)
+					.Text(LOCTEXT("AR_Cancel", "Cancel"))
+				]
 			]
 		];
 }
 
-void SAdvancedRenamerPanel::CreateRightPane(const TSharedRef<SCanvas>& InCanvas)
+void SAdvancedRenamerPanel::CreateRightPanel()
 {
-	using namespace UE::AdvancedRenamer::Private;
+	using namespace AdvancedRenamerSlateUtils::Default;
 
-	InCanvas->AddSlot()
-		.Position(FVector2D(RightBlockOffsetX, 5.f))
-		.Size(this, &SAdvancedRenamerPanel::GetRightPaneSize)
+	RightSideBox->SetContent(SNew(SVerticalBox)
+		+ SVerticalBox::Slot()
 		[
 			CreateRenamePreview()
-		];
+		]);
 }
 
 TSharedRef<SWidget> SAdvancedRenamerPanel::CreateRenamePreview()
 {
-	using namespace UE::AdvancedRenamer::Private;
+	using namespace AdvancedRenamerSlateUtils::Default;
 
 	return SNew(SBorder)
-		.BorderImage(FAppStyle::Get().GetBrush("DetailsView.CategoryTop"))
-		.HAlign(EHorizontalAlignment::HAlign_Fill)
-		.VAlign(EVerticalAlignment::VAlign_Fill)
+		.BorderImage(FAdvancedRenamerStyle::Get().GetBrush("AdvancedRenamer.Style.BackgroundBorder"))
 		.Content()
 		[
 			SNew(SVerticalBox)
 
-			+ SVerticalBox::Slot()
-			.Padding(0.f, LabelOffsetY, 0.f, 0.f)
-			.AutoHeight()
-			[
-				SNew(STextBlock)
-				.ColorAndOpacity(FStyleColors::AccentBlue.GetSpecifiedColor())
-				.Font(TitleFont)
-				.Text(LOCTEXT("ObjectRenamePreviewTitle", "Rename Preview"))
-			]
-
 			+ SVerticalBox::Slot()
 			.FillHeight(1.f)
-			.HAlign(EHorizontalAlignment::HAlign_Fill)
 			[
 				SAssignNew(RenamePreviewListBox, SBox)
 				.Content()
 				[
 					SAssignNew(RenamePreviewList, SListView<TSharedPtr<FAdvancedRenamerPreview>>)
+					.ListViewStyle(&FAdvancedRenamerStyle::Get().GetWidgetStyle<FTableViewStyle>("AdvancedRenamer.Style.ListView"))
 					.ItemHeight(ListLineHeight)
 					.ListItemsSource(&Renamer->GetPreviews())
 					.OnGenerateRow(this, &SAdvancedRenamerPanel::OnGenerateRowForList)
 					.HeaderRow(
+
 						SAssignNew(RenamePreviewListHeaderRow, SHeaderRow)
+						.Style(&FAdvancedRenamerStyle::Get().GetWidgetStyle<FHeaderRowStyle>("AdvancedRenamer.Style.HeaderRow"))
+
 						+ SHeaderRow::Column(OriginalNameColumnName)
-						.DefaultLabel(LOCTEXT("From", "From"))
+						.HeaderContentPadding(FMargin(8.f, 2.f))
+						.DefaultLabel(LOCTEXT("AR_Old", "Old"))
 						.FillWidth(0.5f)
+
 						+ SHeaderRow::Column(NewNameColumnName)
-						.DefaultLabel(LOCTEXT("To", "To"))
+						.HeaderContentPadding(FMargin(8.f, 2.f))
+						.DefaultLabel(LOCTEXT("AR_New", "New"))
 						.FillWidth(0.5f)
 					)
 					.OnKeyDownHandler(this, &SAdvancedRenamerPanel::OnListViewKeyDown)
 					.OnContextMenuOpening(this, &SAdvancedRenamerPanel::GenerateListViewContextMenu)
 				]
 			]
-
-			+ SVerticalBox::Slot()
-			.AutoHeight()
-			.HAlign(EHorizontalAlignment::HAlign_Fill)
-			[
-				SAssignNew(ApplyButton, SButton)
-				.ButtonStyle(FAdvancedRenamerStyle::Get(), "AdvancedRenamer.Style.DarkButton")
-				.ToolTipText(LOCTEXT("ApplyRenameToolTip", "Renames all actors."))
-				.HAlign(EHorizontalAlignment::HAlign_Center)
-				.VAlign(EVerticalAlignment::VAlign_Center)
-				.IsEnabled(this, &SAdvancedRenamerPanel::IsApplyButtonEnabled)
-				.OnClicked(this, &SAdvancedRenamerPanel::OnApplyButtonClicked)
-				[
-					SNew(STextBlock)
-					.ColorAndOpacity(FStyleColors::AccentWhite.GetSpecifiedColor())
-					.Font(RegularFont)
-					.Text(LOCTEXT("ApplyRename", "Apply"))
-				]
-			]
 		];
 }
 
@@ -782,7 +362,7 @@ bool SAdvancedRenamerPanel::CloseWindow()
 	return false;
 }
 
-void SAdvancedRenamerPanel::RefreshListView(const double InCurrentTime)
+void SAdvancedRenamerPanel::RefreshListViewAndUpdate(const double InCurrentTime)
 {
 	const int32 CurrentCount = Renamer->Num();
 
@@ -793,40 +373,11 @@ void SAdvancedRenamerPanel::RefreshListView(const double InCurrentTime)
 		RenamePreviewList->RequestListRefresh();
 	}
 
-	MinDesiredOriginalNameWidth = 0.f;
-	MinDesiredNewNameWidth = 0.f;
 	RenamePreviewList->RebuildList();
 
 	ListLastUpdateTime = InCurrentTime;
 }
 
-void SAdvancedRenamerPanel::UpdateRequiredListWidth()
-{
-	if (!RenamePreviewListBox.IsValid() || !RenamePreviewListHeaderRow.IsValid())
-	{
-		return;
-	}
-
-	if (MinDesiredNewNameWidth == 0.f)
-	{
-		return;
-	}
-
-	using namespace UE::AdvancedRenamer::Private;
-
-	const float ActualWidth = FMath::Max(GetTickSpaceGeometry().GetLocalSize().X, MinDesiredOriginalNameWidth + MinDesiredNewNameWidth + 20.f);
-
-	RenamePreviewListHeaderRow->SetColumnWidth(
-		OriginalNameColumnName, 
-		(MinDesiredOriginalNameWidth + 10.f) / (MinDesiredOriginalNameWidth + MinDesiredNewNameWidth + 20.f) * ActualWidth
-	);
-
-	RenamePreviewListHeaderRow->SetColumnWidth(
-		NewNameColumnName, 
-		(MinDesiredNewNameWidth + 10.f) / (MinDesiredOriginalNameWidth + MinDesiredNewNameWidth + 20.f) * ActualWidth
-	);
-}
-
 void SAdvancedRenamerPanel::RemoveSelectedObjects()
 {
 	if (!RenamePreviewList.IsValid() || RenamePreviewList->GetNumItemsSelected() == 0)
@@ -889,337 +440,10 @@ void SAdvancedRenamerPanel::Tick(const FGeometry& InAllottedGeometry, const doub
 	}
 	else if (Renamer->IsDirty() && InCurrentTime >= (ListLastUpdateTime + MinUpdateFrequency))
 	{
-		RefreshListView(InCurrentTime);
-	}
-	else
-	{
-		UpdateRequiredListWidth();
+		RefreshListViewAndUpdate(InCurrentTime);
 	}
 }
 
-void SAdvancedRenamerPanel::OnBaseNameChanged(const FText& InNewText)
-{
-	Renamer->GetOptions().BaseName = InNewText.ToString();
-	Renamer->MarkDirty();
-}
-
-void SAdvancedRenamerPanel::OnPrefixChanged(const FText& InNewText)
-{
-	Renamer->GetOptions().AddPrefix = InNewText.ToString();
-	Renamer->MarkDirty();
-}
-
-ECheckBoxState SAdvancedRenamerPanel::IsPrefixRemoveChecked() const
-{
-	return bRemovePrefixSeparator ? ECheckBoxState::Checked : ECheckBoxState::Unchecked;
-}
-
-void SAdvancedRenamerPanel::OnPrefixRemoveCheckBoxChanged(ECheckBoxState InNewState)
-{
-	bRemovePrefixSeparator = InNewState == ECheckBoxState::Checked;
-
-	if (bRemovePrefixSeparator)
-	{
-		bRemovePrefixNumChars = false;
-	}
-
-	if (bRemovePrefixSeparator && PrefixSeparatorTextBox.IsValid())
-	{
-		Renamer->GetOptions().RemovePrefixSeparator = PrefixSeparatorTextBox->GetText().ToString();
-		Renamer->GetOptions().RemovePrefixCharacterCount = 0;
-	}
-	else
-	{
-		Renamer->GetOptions().RemovePrefixSeparator = "";
-	}
-
-	Renamer->MarkDirty();
-}
-
-bool SAdvancedRenamerPanel::IsPrefixRemoveSeparatorEnabled() const
-{
-	return bRemovePrefixSeparator;
-}
-
-bool SAdvancedRenamerPanel::OnPrefixSeparatorVerifyTextChanged(const FText& InText, FText& OutErrorText) const
-{
-	if (InText.ToString().Len() > 1)
-	{
-		OutErrorText = LOCTEXT("SeparatorError", "Separators can only be a single character.");
-		return false;
-	}
-
-	return true;
-}
-
-void SAdvancedRenamerPanel::OnPrefixSeparatorChanged(const FText& InNewText)
-{
-	Renamer->GetOptions().RemovePrefixSeparator = InNewText.ToString();
-	Renamer->MarkDirty();
-}
-
-ECheckBoxState SAdvancedRenamerPanel::IsPrefixRemoveCharactersChecked() const
-{
-	return bRemovePrefixNumChars ? ECheckBoxState::Checked : ECheckBoxState::Unchecked;
-}
-
-void SAdvancedRenamerPanel::OnPrefixRemoveCharactersCheckBoxChanged(ECheckBoxState InNewState)
-{
-	bRemovePrefixNumChars = InNewState == ECheckBoxState::Checked;
-
-	if (bRemovePrefixNumChars)
-	{
-		bRemovePrefixSeparator = false;
-	}
-
-	if (bRemovePrefixNumChars && PrefixRemoveCharactersSpinBox.IsValid())
-	{
-		Renamer->GetOptions().RemovePrefixCharacterCount = PrefixRemoveCharactersSpinBox->GetValue();
-		Renamer->GetOptions().RemovePrefixSeparator = "";
-	}
-	else
-	{
-		Renamer->GetOptions().RemovePrefixCharacterCount = 0;
-	}
-
-	Renamer->MarkDirty();
-}
-
-bool SAdvancedRenamerPanel::IsPrefixRemoveNumCharsEnabled() const
-{
-	return bRemovePrefixNumChars;
-}
-
-void SAdvancedRenamerPanel::OnPrefixRemoveCharactersChanged(uint8 InNewValue)
-{
-	Renamer->GetOptions().RemovePrefixCharacterCount = InNewValue;
-	Renamer->MarkDirty();
-}
-
-void SAdvancedRenamerPanel::OnSuffixChanged(const FText& InNewText)
-{
-	Renamer->GetOptions().AddSuffix = InNewText.ToString();
-	Renamer->MarkDirty();
-}
-
-ECheckBoxState SAdvancedRenamerPanel::IsSuffixRemoveChecked() const
-{
-	return bRemoveSuffixSeparator ? ECheckBoxState::Checked : ECheckBoxState::Unchecked;
-}
-
-void SAdvancedRenamerPanel::OnSuffixRemoveCheckBoxChanged(ECheckBoxState InNewState)
-{
-	bRemoveSuffixSeparator = InNewState == ECheckBoxState::Checked;
-
-	if (bRemoveSuffixSeparator)
-	{
-		bRemoveSuffixNumChars = false;
-	}
-
-	if (bRemoveSuffixSeparator && SuffixSeparatorTextBox.IsValid())
-	{
-		Renamer->GetOptions().RemoveSuffixSeparator = SuffixSeparatorTextBox->GetText().ToString();
-		Renamer->GetOptions().RemoveSuffixCharacterCount = 0;
-	}
-	else
-	{
-		Renamer->GetOptions().RemoveSuffixSeparator = "";
-	}
-
-	Renamer->MarkDirty();
-}
-
-bool SAdvancedRenamerPanel::IsSuffixRemoveSeparatorEnabled() const
-{
-	return bRemoveSuffixSeparator;
-}
-
-bool SAdvancedRenamerPanel::OnSuffixSeparatorVerifyTextChanged(const FText& InText, FText& OutErrorText) const
-{
-	if (InText.ToString().Len() > 1)
-	{
-		OutErrorText = LOCTEXT("SeparatorError", "Separators can only be a single character.");
-		return false;
-	}
-
-	return true;
-}
-
-void SAdvancedRenamerPanel::OnSuffixSeparatorChanged(const FText& InNewText)
-{
-	Renamer->GetOptions().RemoveSuffixSeparator = InNewText.ToString();
-	Renamer->MarkDirty();
-}
-
-ECheckBoxState SAdvancedRenamerPanel::IsSuffixRemoveCharactersChecked() const
-{
-	return bRemoveSuffixNumChars ? ECheckBoxState::Checked : ECheckBoxState::Unchecked;
-}
-
-bool SAdvancedRenamerPanel::IsSuffixRemoveNumCharsEnabled() const
-{
-	return bRemoveSuffixNumChars;
-}
-
-void SAdvancedRenamerPanel::OnSuffixRemoveCharactersCheckBoxChanged(ECheckBoxState InNewState)
-{
-	bRemoveSuffixNumChars = InNewState == ECheckBoxState::Checked;
-
-	if (bRemoveSuffixNumChars)
-	{
-		bRemoveSuffixSeparator = false;
-	}
-
-	if (bRemoveSuffixNumChars && SuffixRemoveCharactersSpinBox.IsValid())
-	{
-		Renamer->GetOptions().RemoveSuffixCharacterCount = SuffixRemoveCharactersSpinBox->GetValue();
-		Renamer->GetOptions().RemoveSuffixSeparator = "";
-	}
-	else
-	{
-		Renamer->GetOptions().RemoveSuffixCharacterCount = 0;
-	}
-
-	Renamer->MarkDirty();
-}
-
-void SAdvancedRenamerPanel::OnSuffixRemoveCharactersChanged(uint8 InNewValue)
-{
-	Renamer->GetOptions().RemoveSuffixCharacterCount = InNewValue;
-	Renamer->MarkDirty();
-}
-
-ECheckBoxState SAdvancedRenamerPanel::IsSuffixRemoveNumberChecked() const
-{
-	return Renamer->GetOptions().bRemoveSuffixNumber ? ECheckBoxState::Checked : ECheckBoxState::Unchecked;
-}
-
-void SAdvancedRenamerPanel::OnSuffixRemoveNumberCheckBoxChanged(ECheckBoxState InNewState)
-{
-	Renamer->GetOptions().bRemoveSuffixNumber = InNewState == ECheckBoxState::Checked;
-
-	if (!Renamer->GetOptions().bRemoveSuffixNumber)
-	{
-		Renamer->GetOptions().bAddSuffixNumber = false;
-	}
-
-	Renamer->MarkDirty();
-}
-
-ECheckBoxState SAdvancedRenamerPanel::IsSuffixNumberChecked() const
-{
-	return Renamer->GetOptions().bAddSuffixNumber ? ECheckBoxState::Checked : ECheckBoxState::Unchecked;
-}
-
-bool SAdvancedRenamerPanel::IsSuffixRemoveNumberCheckBoxEnabled() const
-{
-	return !Renamer->GetOptions().bAddSuffixNumber;
-}
-
-void SAdvancedRenamerPanel::OnSuffixNumberCheckBoxChanged(ECheckBoxState InNewState)
-{
-	Renamer->GetOptions().bAddSuffixNumber = InNewState == ECheckBoxState::Checked;
-
-	if (Renamer->GetOptions().bAddSuffixNumber)
-	{
-		Renamer->GetOptions().bRemoveSuffixNumber = true;
-	}
-
-	Renamer->MarkDirty();
-}
-
-void SAdvancedRenamerPanel::OnSuffixNumberStartChanged(int32 InNewValue)
-{
-	Renamer->GetOptions().AddSuffixNumberStart = InNewValue;
-	Renamer->MarkDirty();
-}
-
-void SAdvancedRenamerPanel::OnSuffixNumberStepChanged(int32 InNewValue)
-{
-	Renamer->GetOptions().AddSuffixNumberStep = InNewValue;
-	Renamer->MarkDirty();
-}
-
-ECheckBoxState SAdvancedRenamerPanel::IsSearchReplacePlainTextChecked() const
-{
-	return Renamer->GetOptions().SearchAndReplaceType == EAdvancedRenamerSeachAndReplaceType::PlainText ? ECheckBoxState::Checked : ECheckBoxState::Unchecked;
-}
-
-void SAdvancedRenamerPanel::OnSearchReplacePlainTextCheckBoxChanged(ECheckBoxState InNewState)
-{
-	if (InNewState != ECheckBoxState::Checked)
-	{
-		return;
-	}
-
-	Renamer->GetOptions().SearchAndReplaceType = EAdvancedRenamerSeachAndReplaceType::PlainText;
-	Renamer->MarkDirty();
-}
-
-ECheckBoxState SAdvancedRenamerPanel::IsSearchReplaceRegexChecked() const
-{
-	return Renamer->GetOptions().SearchAndReplaceType == EAdvancedRenamerSeachAndReplaceType::RegularExpression ? ECheckBoxState::Checked : ECheckBoxState::Unchecked;
-}
-
-void SAdvancedRenamerPanel::OnSearchReplaceRegexCheckBoxChanged(ECheckBoxState InNewState)
-{
-	if (InNewState != ECheckBoxState::Checked)
-	{
-		return;
-	}
-
-	Renamer->GetOptions().SearchAndReplaceType = EAdvancedRenamerSeachAndReplaceType::RegularExpression;
-	Renamer->MarkDirty();
-}
-
-ECheckBoxState SAdvancedRenamerPanel::IsSearchReplaceIgnoreCaseChecked() const
-{
-	return Renamer->GetOptions().SearchAndReplaceCase == ESearchCase::IgnoreCase ? ECheckBoxState::Checked : ECheckBoxState::Unchecked;
-}
-
-void SAdvancedRenamerPanel::OnSearchReplaceIgnoreCaseCheckBoxChanged(ECheckBoxState InNewState)
-{
-	switch (InNewState)
-	{
-		case ECheckBoxState::Checked:
-			Renamer->GetOptions().SearchAndReplaceCase = ESearchCase::IgnoreCase;
-			break;
-
-		case ECheckBoxState::Unchecked:
-			Renamer->GetOptions().SearchAndReplaceCase = ESearchCase::CaseSensitive;
-			break;
-
-		default:
-			return;
-	}
-
-	Renamer->MarkDirty();
-}
-
-void SAdvancedRenamerPanel::OnSearchReplaceSearchTextChanged(const FText& InNewText)
-{
-	Renamer->GetOptions().SearchAndReplaceFromText = InNewText.ToString();
-
-	if (Renamer->GetOptions().SearchAndReplaceType == EAdvancedRenamerSeachAndReplaceType::None)
-	{
-		Renamer->GetOptions().SearchAndReplaceType = EAdvancedRenamerSeachAndReplaceType::PlainText;
-	}
-
-	Renamer->MarkDirty();
-}
-
-void SAdvancedRenamerPanel::OnSearchReplaceReplaceTextChanged(const FText& InNewText)
-{
-	Renamer->GetOptions().SearchAndReplaceToText = InNewText.ToString();
-
-	if (Renamer->GetOptions().SearchAndReplaceType == EAdvancedRenamerSeachAndReplaceType::None)
-	{
-		Renamer->GetOptions().SearchAndReplaceType = EAdvancedRenamerSeachAndReplaceType::PlainText;
-	}
-
-	Renamer->MarkDirty();
-}
-
 TSharedRef<ITableRow> SAdvancedRenamerPanel::OnGenerateRowForList(TSharedPtr<FAdvancedRenamerPreview> InItem, 
 	const TSharedRef<STableViewBase>& InOwnerTable)
 {
@@ -1266,50 +490,43 @@ bool SAdvancedRenamerPanel::IsApplyButtonEnabled() const
 
 FReply SAdvancedRenamerPanel::OnApplyButtonClicked()
 {
-	FScopedTransaction Transaction(LOCTEXT("AdvancedRenamerRename", "Advanced Renamer Rename"));
+	FScopedTransaction Transaction(LOCTEXT("AdvancedRenamerRename", "Batch Renamer Rename"));
 
-	if (Renamer->Execute())
+	if (!Renamer->Execute())
 	{
-		CloseWindow();
+		FNotificationInfo Info(LOCTEXT("AR_RenameIssueNotification", "Errors occured while applying the rename.\nSee Output Log for more information on why the rename failed."));
+		Info.ExpireDuration = 5.0f;
+		const TSharedPtr<SNotificationItem> Notification = FSlateNotificationManager::Get().AddNotification(Info);
+		if (Notification.IsValid())
+		{
+			Notification->SetCompletionState(SNotificationItem::CS_Fail);
+		}
 	}
 
+	CloseWindow();
 	return FReply::Handled();
 }
 
-FVector2D SAdvancedRenamerPanel::GetRightPaneSize() const
+FReply SAdvancedRenamerPanel::OnResetButtonClicked()
 {
-	using namespace UE::AdvancedRenamer::Private;
-
-	const FVector2f GeoSize = GetTickSpaceGeometry().GetLocalSize();
-
-	return FVector2D(
-		FMath::Max(RightBlockSize.X, GeoSize.X - 319.0),
-		RightBlockSize.Y
-	);
-}
-
-FVector2D SAdvancedRenamerPanel::GetListViewsize() const
-{
-	using namespace UE::AdvancedRenamer::Private;
+	TArray<TSharedPtr<IAdvancedRenamerSection>> RegisteredSections = IAdvancedRenamerModule::Get().GetRegisteredSections();
+	for (const TSharedPtr<IAdvancedRenamerSection>& Section : RegisteredSections)
+	{
+		Section->ResetToDefault();
+	}
 
-	const FVector2f GeoSize = GetTickSpaceGeometry().GetLocalSize();
+	if (Renamer.IsValid())
+	{
+		Renamer->MarkDirty();
+	}
 
-	return FVector2D(
-		FMath::Max(ListViewSize.X, GeoSize.X - 323.0),
-		ListViewSize.Y
-	);
+	return FReply::Handled();
 }
 
-FVector2D SAdvancedRenamerPanel::GetApplyButtonSize() const
+FReply SAdvancedRenamerPanel::OnCancelButtonClicked()
 {
-	using namespace UE::AdvancedRenamer::Private;
-
-	const FVector2f GeoSize = GetTickSpaceGeometry().GetLocalSize();
-
-	return FVector2D(
-		FMath::Max(ListViewSize.X, GeoSize.X - 323.0),
-		ApplyButtonHeight
-	);
+	CloseWindow();
+	return FReply::Handled();
 }
 
 #undef LOCTEXT_NAMESPACE
diff --git a/Engine/Plugins/Experimental/AdvancedRenamer/Source/AdvancedRenamer/Private/Slate/SAdvancedRenamerPanel.h b/Engine/Plugins/Experimental/AdvancedRenamer/Source/AdvancedRenamer/Private/Slate/SAdvancedRenamerPanel.h
--- a/Engine/Plugins/Experimental/AdvancedRenamer/Source/AdvancedRenamer/Private/Slate/SAdvancedRenamerPanel.h
+++ b/Engine/Plugins/Experimental/AdvancedRenamer/Source/AdvancedRenamer/Private/Slate/SAdvancedRenamerPanel.h
@@ -2,24 +2,23 @@
 
 #pragma once
 
-#include "Styling/SlateTypes.h"
-#include "UObject/StrongObjectPtr.h"
-#include "Widgets/Views/SListView.h"
-#include "Widgets/Views/STableRow.h"
-#include "Widgets/Views/STableViewBase.h"
+#include "Input/Reply.h"
+#include "Widgets/SCompoundWidget.h"
+#include "Templates/SharedPointerFwd.h"
 
 class FUICommandList;
 class IAdvancedRenamer;
+class ITableRow;
 class SBox;
-class SButton;
-class SCanvas;
-class SCheckBox;
-class SEditableTextBox;
 class SHeaderRow;
-class SMultiLineEditableTextBox; 
-class UObject;
+template <typename InItemType>
+class SListView;
+class STableViewBase;
+class SVerticalBox;
+class SWidget;
 struct FAdvancedRenamerPreview;
-template<typename NumericType> class SSpinBox;
+struct FGeometry;
+struct FKeyEvent;
 
 class SAdvancedRenamerPanel : public SCompoundWidget
 {
@@ -31,132 +30,87 @@ public:
 
 	void Construct(const FArguments& InArgs, const TSharedRef<IAdvancedRenamer>& InRenamer);
 
-protected:
-	static constexpr double MinUpdateFrequency = 0.1;
+private:
+	/** Create Apply Reset and Cancel bottom panel */
+	void CreateApplyResetCancelPanel();
 
-	TSharedPtr<IAdvancedRenamer> Renamer;
-	TSharedPtr<FUICommandList> CommandList;
+	/** Create Apply button */
+	TSharedRef<SWidget> CreateApplyButton();
 
-	double ListLastUpdateTime = 0;
-	float MinDesiredOriginalNameWidth = 0.f;
-	float MinDesiredNewNameWidth = 0.f;
-
-	bool bRemovePrefixSeparator;
-	bool bRemovePrefixNumChars;
-	bool bRemoveSuffixSeparator;
-	bool bRemoveSuffixNumChars;
-
-	TSharedPtr<SEditableTextBox> BaseNameTextBox;
-	TSharedPtr<SEditableTextBox> PrefixTextBox;
-	TSharedPtr<SCheckBox> PrefixRemoveCheckBox;
-	TSharedPtr<SEditableTextBox> PrefixSeparatorTextBox;
-	TSharedPtr<SCheckBox> PrefixRemoveCharactersCheckBox;
-	TSharedPtr<SSpinBox<uint8>> PrefixRemoveCharactersSpinBox;
-	TSharedPtr<SEditableTextBox> SuffixTextBox;
-	TSharedPtr<SCheckBox> SuffixRemoveCheckBox;
-	TSharedPtr<SEditableTextBox> SuffixSeparatorTextBox;
-	TSharedPtr<SCheckBox> SuffixRemoveCharactersCheckBox;
-	TSharedPtr<SSpinBox<uint8>> SuffixRemoveCharactersSpinBox;
-	TSharedPtr<SCheckBox> SuffixRemoveNumberCheckBox;
-	TSharedPtr<SCheckBox> SuffixNumberCheckBox;
-	TSharedPtr<SSpinBox<int32>> SuffixNumberStartSpinBox;
-	TSharedPtr<SSpinBox<int32>> SuffixNumberStepSpinBox;
-	TSharedPtr<SCheckBox> SearchReplacePlainTextCheckbox;
-	TSharedPtr<SCheckBox> SearchReplaceRegexCheckbox;
-	TSharedPtr<SCheckBox> SearchReplaceIgnoreCaseCheckBox;
-	TSharedPtr<SMultiLineEditableTextBox> SearchReplaceSearchTextBox;
-	TSharedPtr<SMultiLineEditableTextBox> SearchReplaceReplaceTextBox;
-	TSharedPtr<SBox> RenamePreviewListBox;
-	TSharedPtr<SHeaderRow> RenamePreviewListHeaderRow;
-	TSharedPtr<SListView<TSharedPtr<FAdvancedRenamerPreview>>> RenamePreviewList;
-	TSharedPtr<SButton> ApplyButton;
+	/** Create Reset button */
+	TSharedRef<SWidget> CreateResetButton();
+
+	/** Create Cancel button */
+	TSharedRef<SWidget> CreateCancelButton();
 
-	void CreateLeftPane(const TSharedRef<SCanvas>& InCanvas);
-	TSharedRef<SWidget> CreateBaseName();
-	TSharedRef<SWidget> CreatePrefix();
-	TSharedRef<SWidget> CreateSuffix();
-	TSharedRef<SWidget> CreateSearchAndReplace();
+	/** Create the right side panel */
+	void CreateRightPanel();
 
-	void CreateRightPane(const TSharedRef<SCanvas>& InCanvas);
+	/** Create the rename preview */
 	TSharedRef<SWidget> CreateRenamePreview();
 
+	/** Close the Renamer window */
 	bool CloseWindow();
 
-	void RefreshListView(const double InCurrentTime);
-	void UpdateRequiredListWidth();
+	/** Refresh and Update the Preview */
+	void RefreshListViewAndUpdate(const double InCurrentTime);
 
+	/** Remove selected entries from the Renamer*/
 	void RemoveSelectedObjects();
 
+	//~ Begin SWidget Interface
 	virtual void Tick(const FGeometry& InAllottedGeometry, const double InCurrentTime, const float InDeltaTime) override;
+	//~ End SWidget Interface
 
-	void OnBaseNameChanged(const FText& InNewText);
-
-	void OnPrefixChanged(const FText& InNewText);
-
-	ECheckBoxState IsPrefixRemoveChecked() const;
-	void OnPrefixRemoveCheckBoxChanged(ECheckBoxState InNewState);
-
-	bool IsPrefixRemoveSeparatorEnabled() const;
-	bool OnPrefixSeparatorVerifyTextChanged(const FText& InText, FText& OutErrorText) const;
-	void OnPrefixSeparatorChanged(const FText& InNewText);
-
-	ECheckBoxState IsPrefixRemoveCharactersChecked() const;
-	void OnPrefixRemoveCharactersCheckBoxChanged(ECheckBoxState InNewState);
-
-	bool IsPrefixRemoveNumCharsEnabled() const;
-	void OnPrefixRemoveCharactersChanged(uint8 InNewValue);
-
-	void OnSuffixChanged(const FText& InNewText);
-
-	ECheckBoxState IsSuffixRemoveChecked() const;
-	void OnSuffixRemoveCheckBoxChanged(ECheckBoxState InNewState);
-
-	bool IsSuffixRemoveSeparatorEnabled() const;
-	bool OnSuffixSeparatorVerifyTextChanged(const FText& InText, FText& OutErrorText) const;
-	void OnSuffixSeparatorChanged(const FText& InNewText);
-
-	ECheckBoxState IsSuffixRemoveCharactersChecked() const;
-	void OnSuffixRemoveCharactersCheckBoxChanged(ECheckBoxState InNewState);
-
-	bool IsSuffixRemoveNumCharsEnabled() const;
-	void OnSuffixRemoveCharactersChanged(uint8 InNewValue);
+	/** Callback to generate ListBoxRows */
+	TSharedRef<ITableRow> OnGenerateRowForList(TSharedPtr<FAdvancedRenamerPreview> InItem, const TSharedRef<STableViewBase>& InOwnerTable);
 
-	ECheckBoxState IsSuffixRemoveNumberChecked() const;
-	void OnSuffixRemoveNumberCheckBoxChanged(ECheckBoxState InNewState);
+	/** Handle the click on the ListView */
+	FReply OnListViewKeyDown(const FGeometry& InGeometry, const FKeyEvent& InKeyEvent);
 
-	ECheckBoxState IsSuffixNumberChecked() const;
-	bool IsSuffixRemoveNumberCheckBoxEnabled() const;
-	void OnSuffixNumberCheckBoxChanged(ECheckBoxState InNewState);
+	/** Generate the ListView menu */
+	TSharedPtr<SWidget> GenerateListViewContextMenu();
 
-	void OnSuffixNumberStartChanged(int32 InNewValue);
+	/** Whether or not the ApplyButton is enabled */
+	bool IsApplyButtonEnabled() const;
 
-	void OnSuffixNumberStepChanged(int32 InNewValue);
+	/** Called when clicking on the Apply button */
+	FReply OnApplyButtonClicked();
 
-	ECheckBoxState IsSearchReplacePlainTextChecked() const;
-	void OnSearchReplacePlainTextCheckBoxChanged(ECheckBoxState InNewState);
+	/** Called when clicking on the Reset button */
+	FReply OnResetButtonClicked();
 
-	ECheckBoxState IsSearchReplaceRegexChecked() const;
-	void OnSearchReplaceRegexCheckBoxChanged(ECheckBoxState InNewState);
+	/** Called when clicking on the Cancel button */
+	FReply OnCancelButtonClicked();
 
-	ECheckBoxState IsSearchReplaceIgnoreCaseChecked() const;
-	void OnSearchReplaceIgnoreCaseCheckBoxChanged(ECheckBoxState InNewState);
+private:
+	/** Min update frequency, used in the tick to avoid updating the renamer too often */
+	static constexpr double MinUpdateFrequency = 0.1f;
 
-	void OnSearchReplaceSearchTextChanged(const FText& InNewText);
+	/** Renamer instance */
+	TSharedPtr<IAdvancedRenamer> Renamer;
 
-	void OnSearchReplaceReplaceTextChanged(const FText& InNewText);
+	/** Command list of the Renamer Panel */
+	TSharedPtr<FUICommandList> CommandList;
 
-	TSharedRef<ITableRow> OnGenerateRowForList(TSharedPtr<FAdvancedRenamerPreview> InItem, const TSharedRef<STableViewBase>& InOwnerTable);
+	/** Last update time of the list */
+	double ListLastUpdateTime = 0;
 
-	FReply OnListViewKeyDown(const FGeometry& InGeometry, const FKeyEvent& InKeyEvent);
+	/** Sections VerticalBox */
+	TSharedPtr<SVerticalBox> LeftSideVerticalBox;
 
-	TSharedPtr<SWidget> GenerateListViewContextMenu();
+	/** RightSide container Box */
+	TSharedPtr<SBox> RightSideBox;
 
-	bool IsApplyButtonEnabled() const;
-	FReply OnApplyButtonClicked();
+	/** Bottom panel button container Box */
+	TSharedPtr<SBox> ApplyResetCancelBox;
 
-	FVector2D GetRightPaneSize() const;
+	/** ListView container Box */
+	TSharedPtr<SBox> RenamePreviewListBox;
 
-	FVector2D GetListViewsize() const;
+	/** ListView Header row */
+	TSharedPtr<SHeaderRow> RenamePreviewListHeaderRow;
 
-	FVector2D GetApplyButtonSize() const;
+	/** ListView of the Previews */
+	TSharedPtr<SListView<TSharedPtr<FAdvancedRenamerPreview>>> RenamePreviewList;
 };
diff --git a/Engine/Plugins/Experimental/AdvancedRenamer/Source/AdvancedRenamer/Private/Utils/AdvancedRenamerSlateUtils.h b/Engine/Plugins/Experimental/AdvancedRenamer/Source/AdvancedRenamer/Private/Utils/AdvancedRenamerSlateUtils.h
new file mode 100644
--- /dev/null
+++ b/Engine/Plugins/Experimental/AdvancedRenamer/Source/AdvancedRenamer/Private/Utils/AdvancedRenamerSlateUtils.h
@@ -0,0 +1,35 @@
+﻿// Copyright Epic Games, Inc. All Rights Reserved.
+
+#pragma once
+
+#include "Fonts/SlateFontInfo.h"
+#include "Layout/Margin.h"
+#include "Styling/CoreStyle.h"
+#include "UObject/NameTypes.h"
+
+enum class EAdvancedRenamerRemoveOldType : uint8
+{
+	Separator,
+	Chars
+};
+
+namespace AdvancedRenamerSlateUtils::Default
+{
+	const FSlateFontInfo TitleFont = FCoreStyle::GetDefaultFontStyle("Regular", 12);
+	const FSlateFontInfo RegularFont = FCoreStyle::GetDefaultFontStyle("Regular", 10);
+	constexpr float ListLineHeight = 15.f;
+	static FName OriginalNameColumnName = TEXT("OriginalName");
+	static FName NewNameColumnName = TEXT("NewName");
+	static FMargin NameWidgetPadding = FMargin(0.f, 0.f, 4.f, 0.f);
+	static FMargin ValueWidgetPadding = FMargin(4.f, 0.f, 0.f, 0.f);
+	static FMargin SectionContentFirstEntryPadding = FMargin(8.f, 8.f);
+	static FMargin SectionContentMiddleEntriesPadding = FMargin(8.f, 0.f, 8.f, 8.f);
+	static FMargin VerticalAddNumberPadding = FMargin(0.f, 4.f);
+	static FMargin AddNumberStepPadding = FMargin(8.f, 0.f, 0.f, 0.f);
+	static FMargin ChangeCaseFirstButtonPadding = FMargin(0.f, 0.f, 4.f, 0.f);
+	static FMargin ChangeCaseMiddleButtonsPadding = FMargin(4.f, 0.f);
+	static FMargin ChangeCaseLastButtonPadding = FMargin(4.f, 0.f, 0.f, 0.f);
+	static FMargin ApplyButtonPadding = FMargin(8.f, 8.f, 4.f, 8.f);
+	static FMargin ResetButtonPadding = FMargin(4.f, 8.f, 4.f, 8.f);
+	static FMargin CancelButtonPadding = FMargin(4.f, 8.f, 8.f, 8.f);
+}
diff --git a/Engine/Plugins/Experimental/AdvancedRenamer/Source/AdvancedRenamer/Public/IAdvancedRenamerModule.h b/Engine/Plugins/Experimental/AdvancedRenamer/Source/AdvancedRenamer/Public/IAdvancedRenamerModule.h
--- a/Engine/Plugins/Experimental/AdvancedRenamer/Source/AdvancedRenamer/Public/IAdvancedRenamerModule.h
+++ b/Engine/Plugins/Experimental/AdvancedRenamer/Source/AdvancedRenamer/Public/IAdvancedRenamerModule.h
@@ -9,6 +9,7 @@
 class AActor;
 class IAdvancedRenamer;
 class IAdvancedRenamerProvider;
+class IAdvancedRenamerSection;
 class IToolkitHost;
 class SWidget;
 
@@ -43,5 +44,7 @@ public:
 
 	virtual void OpenAdvancedRenamer(const TSharedRef<IAdvancedRenamer>& InRenamer, const TSharedPtr<SWidget>& InParentWidget) = 0;
 
+	virtual TArray<TSharedPtr<IAdvancedRenamerSection>> GetRegisteredSections() const = 0;
+
 	virtual TArray<AActor*> GetActorsSharingClassesInWorld(const TArray<AActor*>& InActors) = 0;
 };
