[Backout] - CL33300102
Original CL Desc
-----------------------------------------------------------------
[AdvancedRenamer] Re-Style
Updating AdvancedRenamer style plugin to be more consistent with the rest of the engine.
Updating AdvancedRenamer backend to be more flexible and extensible in the future.
Renaming AdvancedRenamer to BatchRenamer in friendly uplugin name and UI.

diff --git a/Engine/Plugins/Experimental/AdvancedRenamer/AdvancedRenamer.uplugin b/Engine/Plugins/Experimental/AdvancedRenamer/AdvancedRenamer.uplugin
--- a/Engine/Plugins/Experimental/AdvancedRenamer/AdvancedRenamer.uplugin
+++ b/Engine/Plugins/Experimental/AdvancedRenamer/AdvancedRenamer.uplugin
@@ -2,8 +2,8 @@
 	"FileVersion": 3,
 	"Version": 1,
 	"VersionName": "1.0",
-	"FriendlyName": "Batch Renamer",
-	"Description": "Batch Renaming tool for actors and assets",
+	"FriendlyName": "Advanced Renamer",
+	"Description": "Advanced Renaming tool for actors and assets",
 	"Category": "Experimental",
 	"CreatedBy": "Epic Games, Inc.",
 	"CreatedByURL": "https://epicgames.com",
@@ -13,7 +13,6 @@
 	"CanContainContent": false,
 	"IsBetaVersion": false,
 	"IsExperimentalVersion": true,
-	"EnabledByDefault": true,
 	"Installed": false,
 	"Modules": [
 		{
diff --git a/Engine/Plugins/Experimental/AdvancedRenamer/Source/AdvancedRenamer/Private/AdvancedRenamer.cpp b/Engine/Plugins/Experimental/AdvancedRenamer/Source/AdvancedRenamer/Private/AdvancedRenamer.cpp
--- a/Engine/Plugins/Experimental/AdvancedRenamer/Source/AdvancedRenamer/Private/AdvancedRenamer.cpp
+++ b/Engine/Plugins/Experimental/AdvancedRenamer/Source/AdvancedRenamer/Private/AdvancedRenamer.cpp
@@ -3,10 +3,16 @@
 #include "AdvancedRenamer.h"
 
 #include "AdvancedRenamerModule.h"
+#include "Internationalization/Regex.h"
 
-FAdvancedRenamer::FAdvancedRenamer(const TSharedRef<IAdvancedRenamerProvider>& InProvider)
+FAdvancedRenamer::FAdvancedRenamer(const TSharedRef<IAdvancedRenamerProvider>& InProvider, FAdvancedRenamerOptions* InInitialOptions)
 	: Provider(InProvider)
 {
+	if (InInitialOptions)
+	{
+		SetOption(*InInitialOptions);
+	}
+
 	int32 Count = Num();
 	check(Count > 0);
 
@@ -32,6 +38,22 @@ const TSharedRef<IAdvancedRenamerProvider>& FAdvancedRenamer::GetProvider() cons
 	return Provider;
 }
 
+const FAdvancedRenamerOptions& FAdvancedRenamer::GetOptions() const
+{
+	return Options;
+}
+
+FAdvancedRenamerOptions& FAdvancedRenamer::GetOptions()
+{
+	return Options;
+}
+
+void FAdvancedRenamer::SetOption(const FAdvancedRenamerOptions& InOptions)
+{
+	Options = InOptions;
+	MarkDirty();
+}
+
 const TArray<TSharedPtr<FAdvancedRenamerPreview>>& FAdvancedRenamer::GetPreviews()
 {
 	return Previews;
@@ -47,11 +69,6 @@ TSharedPtr<FAdvancedRenamerPreview> FAdvancedRenamer::GetPreview(int32 InIndex)
 	return nullptr;
 }
 
-void FAdvancedRenamer::AddSection(FAdvancedRenamerExecuteSection InSection)
-{
-	Sections.Add(InSection);
-}
-
 bool FAdvancedRenamer::HasRenames() const
 {
 	return bHasRenames;
@@ -77,8 +94,6 @@ bool FAdvancedRenamer::UpdatePreviews()
 	bHasRenames = false;
 	const int32 Count = Previews.Num();
 
-	BeforeOperationsStartExecute();
-
 	for (int32 Index = 0; Index < Count; ++Index)
 	{
 		if (!Previews[Index].IsValid() || !IsValidIndex(Index))
@@ -89,14 +104,14 @@ bool FAdvancedRenamer::UpdatePreviews()
 		}
 
 		// Force recreation
-		Previews[Index]->NewName = ApplyRename(Previews[Index]->OriginalName);
+		Previews[Index]->NewName = ApplyRename(Previews[Index]->OriginalName, Index);
 
 		if (Previews[Index]->NewName.IsEmpty())
 		{
 			continue;
 		}
 
-		if (GetOriginalName(Index).Equals(Previews[Index]->NewName))
+		if (GetOriginalName(Index) == Previews[Index]->NewName)
 		{
 			continue;
 		}
@@ -104,8 +119,6 @@ bool FAdvancedRenamer::UpdatePreviews()
 		bHasRenames = true;
 	}
 
-	AfterOperationsEndExecute();
-
 	MarkClean();
 
 	return bHasRenames;
@@ -154,30 +167,283 @@ bool FAdvancedRenamer::Execute()
 	return bAllSuccess;
 }
 
-FString FAdvancedRenamer::ApplyRename(const FString& InOriginalName)
+FString FAdvancedRenamer::ApplyRename(const FString& InOriginalName, int32 InIndex) const
 {
 	FString NewName = InOriginalName;
-	for (FAdvancedRenamerExecuteSection& Section : Sections)
+
+	if (Options.HasBaseName())
+	{
+		NewName = ApplyBaseName(NewName);
+	}
+
+	if (Options.HasPrefix())
+	{
+		NewName = ApplyPrefix(NewName);
+	}
+
+	if (Options.HasSuffix())
+	{
+		NewName = ApplySuffix(NewName, InIndex);
+	}
+
+	if (Options.HasSearchAndReplace())
 	{
-		Section.OnOperationExecuted().ExecuteIfBound(NewName);
+		switch (Options.SearchAndReplaceType)
+		{
+			default:
+				break;
+
+			case EAdvancedRenamerSeachAndReplaceType::PlainText:
+				NewName = ApplySearchPlainText(NewName);
+				break;
+
+			case EAdvancedRenamerSeachAndReplaceType::RegularExpression:
+				NewName = ApplySearchReplaceRegex(NewName);
+				break;
+		}
 	}
+
 	return NewName;
 }
 
-void FAdvancedRenamer::BeforeOperationsStartExecute()
+FString FAdvancedRenamer::ApplyBaseName(const FString& InOriginalName) const
 {
-	for (FAdvancedRenamerExecuteSection& Section : Sections)
+	if (Options.BaseName.IsEmpty())
 	{
-		Section.OnBeforeOperationExecutionStart().ExecuteIfBound();
+		return InOriginalName;
 	}
+
+	return Options.BaseName;
 }
 
-void FAdvancedRenamer::AfterOperationsEndExecute()
+FString FAdvancedRenamer::ApplyPrefix(const FString& InOriginalName) const
 {
-	for (FAdvancedRenamerExecuteSection& Section : Sections)
+	FString Output = InOriginalName;
+
+	if (!Options.RemovePrefixSeparator.IsEmpty())
 	{
-		Section.OnAfterOperationExecutionEnded().ExecuteIfBound();
+		const int32 PrefixStart = Output.Find(Options.RemovePrefixSeparator, ESearchCase::IgnoreCase);
+
+		if (PrefixStart >= 0)
+		{
+			Output = Output.Mid(PrefixStart + Options.RemovePrefixSeparator.Len());
+		}
 	}
+	else if (Options.RemovePrefixCharacterCount > 0)
+	{
+		Output = Output.RightChop(Options.RemovePrefixCharacterCount);
+	}
+
+	if (!Options.AddPrefix.IsEmpty())
+	{
+		Output = Options.AddPrefix + Output;
+	}
+
+	return Output;
+}
+
+FString FAdvancedRenamer::ApplySuffix(const FString& InOriginalName, int32 InIndex) const
+{
+	static constexpr TCHAR FirstDigit = '0';
+	static constexpr TCHAR LastDigit = '9';
+
+	FString Output = InOriginalName;
+
+	if (!Options.RemoveSuffixSeparator.IsEmpty())
+	{
+		const int32 SuffixStart = Output.Find(Options.RemoveSuffixSeparator, ESearchCase::IgnoreCase, ESearchDir::FromEnd);
+
+		if (SuffixStart >= 0)
+		{
+			Output = Output.Mid(0, SuffixStart);
+		}
+	}
+	else if (Options.RemoveSuffixCharacterCount > 0)
+	{
+		Output = Output.LeftChop(Options.RemoveSuffixCharacterCount);
+	}
+
+	if (Options.bRemoveSuffixNumber || Options.bAddSuffixNumber)
+	{
+		int32 LastDigitIndex = Output.Len() - 1;
+
+		while (LastDigitIndex >= 0)
+		{
+			if (Output[LastDigitIndex] < FirstDigit || Output[LastDigitIndex] > LastDigit)
+			{
+				break;
+			}
+
+			--LastDigitIndex;
+		}
+
+		Output = Output.Mid(0, LastDigitIndex + 1);
+	}
+
+	if (!Options.AddSuffix.IsEmpty())
+	{
+		Output = Output + Options.AddSuffix;
+	}
+
+	if (Options.bAddSuffixNumber)
+	{
+		const int32 Current = Options.AddSuffixNumberStart + (InIndex * Options.AddSuffixNumberStep);
+		Output += FString::FromInt(Current);
+	}
+
+	return Output;
+}
+
+FString FAdvancedRenamer::ApplySearchPlainText(const FString& InOriginalName) const
+{
+	if (Options.SearchAndReplaceFromText.IsEmpty())
+	{
+		return InOriginalName;
+	}
+
+	return InOriginalName.Replace(
+		*Options.SearchAndReplaceFromText, 
+		*Options.SearchAndReplaceToText, 
+		Options.SearchAndReplaceCase
+	);
+}
+
+FString FAdvancedRenamer::ApplySearchReplaceRegex(const FString& InOriginalName) const
+{
+	if (Options.SearchAndReplaceFromText.IsEmpty())
+	{
+		return InOriginalName;
+	}
+
+	const FRegexPattern RegexPattern = FRegexPattern(
+		Options.SearchAndReplaceFromText,
+		Options.SearchAndReplaceCase == ESearchCase::IgnoreCase ? ERegexPatternFlags::CaseInsensitive : ERegexPatternFlags::None
+	);
+
+	return RegexReplace(InOriginalName, RegexPattern, Options.SearchAndReplaceToText);
+}
+
+FString FAdvancedRenamer::RegexReplace(const FString& InOriginalString, const FRegexPattern& InPattern, const FString& InReplaceString) const
+{
+	static const FString EscapeString = TEXT("\\");
+	static constexpr TCHAR EscapeChar = '\\';
+	static const FString GroupString = TEXT("$");
+	static constexpr TCHAR GroupChar = '$';
+	static constexpr TCHAR FirstDigit = '0';
+	static constexpr TCHAR LastDigit = '9';
+	static constexpr TCHAR NullChar = 0;
+
+	FRegexMatcher Matcher(InPattern, InOriginalString);
+
+	FString Output = "";
+	int32 StartCharIndex = 0;
+	bool bReadingGroupName = false;
+	int32 GroupIndex = INDEX_NONE;
+
+	while (Matcher.FindNext())
+	{
+		// Add on part of string after start/previous match
+		if (StartCharIndex != Matcher.GetMatchBeginning())
+		{
+			Output += InOriginalString.Mid(StartCharIndex, Matcher.GetMatchBeginning() - StartCharIndex);
+		}
+
+		bool bEscaped = false;
+
+		for (int32 CharIndex = 0; CharIndex <= InReplaceString.Len(); ++CharIndex)
+		{
+			const TCHAR& Char = CharIndex < InReplaceString.Len() ? InReplaceString[CharIndex] : NullChar;
+
+			if (bReadingGroupName)
+			{
+				// Build group index
+				if (Char >= FirstDigit && Char <= LastDigit)
+				{
+					if (GroupIndex == INDEX_NONE)
+					{
+						GroupIndex = 0;
+					}
+					else if (GroupIndex > 0)
+					{
+						GroupIndex *= 10;
+					}
+
+					const int32 NextDigit = static_cast<int32>(Char - FirstDigit);
+					GroupIndex += NextDigit;
+					continue;
+				}
+				// We've read a group index, add it to the output string
+				else if (GroupIndex > 0)
+				{
+					if (Matcher.GetCaptureGroupBeginning(GroupIndex) == INDEX_NONE)
+					{
+						UE_LOG(LogARP, Error, TEXT("Regex: Capture group does not exist %d."), GroupIndex);
+					}
+
+					Output += Matcher.GetCaptureGroup(GroupIndex);
+				}
+				// $0 matches the entire matched string
+				else if (GroupIndex == 0)
+				{
+					Output += InOriginalString.Mid(Matcher.GetMatchBeginning(), Matcher.GetMatchEnding() - Matcher.GetMatchBeginning());
+				}
+				// An unescaped $
+				else
+				{
+					UE_LOG(LogARP, Error, TEXT("Regex: Unescaped %s."), *GroupString);
+
+					Output += GroupString;
+				}
+
+				bReadingGroupName = false;
+				// Continue regular parsing of this character.
+			}
+
+			// Check for special chars
+			if (!bEscaped)
+			{
+				if (Char == EscapeChar)
+				{
+					bEscaped = true;
+					continue;
+				}
+
+				if (Char == GroupChar)
+				{
+					bReadingGroupName = true;
+					GroupIndex = INDEX_NONE;
+					continue;
+				}
+			}
+			else
+			{
+				// If the last char is a \ assume that it's not an escape char
+				if (Char == NullChar)
+				{
+					UE_LOG(LogARP, Error, TEXT("Regex: Unescaped %s."), *EscapeString);
+
+					Output += EscapeChar;
+				}
+			}
+
+			if (Char != NullChar)
+			{
+				Output += Char;
+			}
+
+			bEscaped = false;
+		}
+
+		StartCharIndex = Matcher.GetMatchEnding();
+	}
+
+	// Add on the end of the string after the last match
+	if (StartCharIndex < InOriginalString.Len())
+	{
+		Output += InOriginalString.Mid(StartCharIndex);
+	}
+
+	return Output;
 }
 
 int32 FAdvancedRenamer::Num() const
diff --git a/Engine/Plugins/Experimental/AdvancedRenamer/Source/AdvancedRenamer/Private/AdvancedRenamer.h b/Engine/Plugins/Experimental/AdvancedRenamer/Source/AdvancedRenamer/Private/AdvancedRenamer.h
--- a/Engine/Plugins/Experimental/AdvancedRenamer/Source/AdvancedRenamer/Private/AdvancedRenamer.h
+++ b/Engine/Plugins/Experimental/AdvancedRenamer/Source/AdvancedRenamer/Private/AdvancedRenamer.h
@@ -6,23 +6,27 @@
 #include "Providers/IAdvancedRenamerProvider.h"
 #include "Templates/SharedPointer.h"
 
+class FRegexPattern;
+struct FAdvancedRenamerOptions;
 struct FAdvancedRenamerPreview;
 
 class FAdvancedRenamer : public IAdvancedRenamer
 {
 public:
-	FAdvancedRenamer(const TSharedRef<IAdvancedRenamerProvider>& InProvider);
+	FAdvancedRenamer(const TSharedRef<IAdvancedRenamerProvider>& InProvider, FAdvancedRenamerOptions* InInitialOptions = nullptr);
 
 	//~ Begin IAdvancedRenamer
 	virtual const TSharedRef<IAdvancedRenamerProvider>& GetProvider() const override;
+	virtual const FAdvancedRenamerOptions& GetOptions() const override;
+	virtual FAdvancedRenamerOptions& GetOptions() override;
+	virtual void SetOption(const FAdvancedRenamerOptions& InOptions) override;
 	virtual const TArray<TSharedPtr<FAdvancedRenamerPreview>>& GetPreviews() override;
 	virtual TSharedPtr<FAdvancedRenamerPreview> GetPreview(int32 InIndex) const override;
-	virtual void AddSection(FAdvancedRenamerExecuteSection InSection) override;
 	virtual bool HasRenames() const override;
 	virtual bool IsDirty() const override;
 	virtual void MarkDirty() override;
 	virtual void MarkClean() override;
-	virtual FString ApplyRename(const FString& InName) override;
+	virtual FString ApplyRename(const FString& InName, int32 InIndex) const override;
 	virtual bool UpdatePreviews() override;
 	virtual bool Execute() override;
 	//~ End IAdvancedRenamer
@@ -37,26 +41,18 @@ public:
 	virtual bool ExecuteRename(int32 InIndex, const FString& InNewName) override;
 	//~ End IAdvancedRenamerProvider
 
-private:
-	/** Called before the whole Rename logic start */
-	void BeforeOperationsStartExecute();
-	
-	/** Called after the whole Rename logic end */
-	void AfterOperationsEndExecute();
-
-private:
-	/** Provider for this Renamer */
+protected:
 	TSharedRef<IAdvancedRenamerProvider> Provider;
-
-	/** Previews Name list */
 	TArray<TSharedPtr<FAdvancedRenamerPreview>> Previews;
-	
-	/** Renamer sections list */
-	TArray<FAdvancedRenamerExecuteSection> Sections;
-
-	/** Whether or not at least 1 preview has a Rename */
+	FAdvancedRenamerOptions Options;
 	bool bHasRenames;
-	
-	/** if true the Rename logic will be executed */
 	bool bDirty;
+
+	FString ApplyBaseName(const FString& InOriginalName) const;
+	FString ApplyPrefix(const FString& InOriginalName) const;
+	FString ApplySuffix(const FString& InOriginalName, int32 InIndex) const;
+	FString ApplySearchPlainText(const FString& InOriginalName) const;
+	FString ApplySearchReplaceRegex(const FString& InOriginalName) const;
+
+	FString RegexReplace(const FString& InOriginalString, const FRegexPattern& InPattern, const FString& InReplaceString) const;
 };
diff --git a/Engine/Plugins/Experimental/AdvancedRenamer/Source/AdvancedRenamer/Private/AdvancedRenamerCommands.cpp b/Engine/Plugins/Experimental/AdvancedRenamer/Source/AdvancedRenamer/Private/AdvancedRenamerCommands.cpp
--- a/Engine/Plugins/Experimental/AdvancedRenamer/Source/AdvancedRenamer/Private/AdvancedRenamerCommands.cpp
+++ b/Engine/Plugins/Experimental/AdvancedRenamer/Source/AdvancedRenamer/Private/AdvancedRenamerCommands.cpp
@@ -22,13 +22,13 @@ void FAdvancedRenamerCommands::RegisterCommands()
 {
 	UI_COMMAND(RenameSelectedActors
 		, "Rename Selected Actors"
-		, "Opens the Batch Renamer Panel to rename all selected actors."
+		, "Opens the Advanced Renamer Panel to rename all selected actors."
 		, EUserInterfaceActionType::Button
 		, FInputChord())
 
 	UI_COMMAND(RenameSharedClassActors
 		, "Rename Actors of Selected Actor Classes"
-		, "Opens the Batch Renamer Panel to rename all actors sharing a class with any selected actor."
+		, "Opens the Advanced Renamer Panel to rename all actors sharing a class with any selected actor."
 		, EUserInterfaceActionType::Button
 		, FInputChord())
 }
diff --git a/Engine/Plugins/Experimental/AdvancedRenamer/Source/AdvancedRenamer/Private/AdvancedRenamerExecuteSection.h b/Engine/Plugins/Experimental/AdvancedRenamer/Source/AdvancedRenamer/Private/AdvancedRenamerExecuteSection.h
deleted file mode 100644
diff --git a/Engine/Plugins/Experimental/AdvancedRenamer/Source/AdvancedRenamer/Private/AdvancedRenamerModule.cpp b/Engine/Plugins/Experimental/AdvancedRenamer/Source/AdvancedRenamer/Private/AdvancedRenamerModule.cpp
--- a/Engine/Plugins/Experimental/AdvancedRenamer/Source/AdvancedRenamer/Private/AdvancedRenamerModule.cpp
+++ b/Engine/Plugins/Experimental/AdvancedRenamer/Source/AdvancedRenamer/Private/AdvancedRenamerModule.cpp
@@ -4,11 +4,6 @@
 
 #include "AdvancedRenamer.h"
 #include "AdvancedRenamerCommands.h"
-#include "AdvancedRenamerSections/AdvancedRenamerAddPrefixSuffixSection.h"
-#include "AdvancedRenamerSections/AdvancedRenamerChangeCaseSection.h"
-#include "AdvancedRenamerSections/AdvancedRenamerRemovePrefixSection.h"
-#include "AdvancedRenamerSections/AdvancedRenamerRemoveSuffixSection.h"
-#include "AdvancedRenamerSections/AdvancedRenamerSearchAndReplaceSection.h"
 #include "AdvancedRenamerStyle.h"
 #include "EngineUtils.h"
 #include "Framework/Application/SlateApplication.h"
@@ -28,59 +23,23 @@ namespace UE::AdvancedRenamer::Private
 	TSharedRef<SWindow> CreateAdvancedRenamerWindow()
 	{
 		return SNew(SWindow)
-			.Title(LOCTEXT("AdvancedRenameWindow", "Batch Renamer"))
-			.ClientSize(FVector2D(731, 547))
+			.Title(LOCTEXT("AdvancedRenameWindow", "Rename Actors"))
+			.ClientSize(FVector2D(600, 538))
 			.SizingRule(ESizingRule::UserSized)
 			.SupportsMaximize(false)
 			.SupportsMinimize(false)
-			.MinWidth(737)
-			.MinHeight(586.5f);
-	}
-
-	static TAutoConsoleVariable<bool> CVarBatchRenamerEnableIntegration(
-		TEXT("BatchRenamer.EnableIntegration"),
-		false,
-		TEXT("If enabled, the BatchRenamer will hook into the ContentBrowser and Outliner.")
-	);
-
-	static bool bPreviousCVarValue = false;
-
-	void OnBatchRenamerEnabledChanged(IConsoleVariable* InVar)
-	{
-		if (InVar)
-		{
-			const bool bCurrentValue = InVar->GetBool();
-			if (bCurrentValue != bPreviousCVarValue)
-			{
-				if (InVar->GetBool())
-				{
-					FAdvancedRenamerContentBrowserIntegration::Initialize();
-					FAdvancedRenamerLevelEditorIntegration::InitializeMenu();
-				}
-				else
-				{
-					FAdvancedRenamerContentBrowserIntegration::Shutdown();
-					FAdvancedRenamerLevelEditorIntegration::ShutdownMenu();
-				}
-				bPreviousCVarValue = bCurrentValue;
-			}
-		}
+			.MinWidth(600)
+			.MinHeight(538)
+			.MaxHeight(538);
 	}
 }
 
 void FAdvancedRenamerModule::StartupModule()
 {
 	FAdvancedRenamerStyle::Initialize();
-	FAdvancedRenamerLevelEditorIntegration::Initialize();
 	FAdvancedRenamerCommands::Register();
-
-	{
-		// Enable plugins integration section from CVar
-		using namespace UE::AdvancedRenamer::Private;
-		EnableRenamerHandle = CVarBatchRenamerEnableIntegration.AsVariable()->OnChangedDelegate().AddStatic(&OnBatchRenamerEnabledChanged);
-	}
-
-	RegisterDefaultSections();
+	FAdvancedRenamerContentBrowserIntegration::Initialize();
+	FAdvancedRenamerLevelEditorIntegration::Initialize();
 }
 
 void FAdvancedRenamerModule::ShutdownModule()
@@ -89,12 +48,6 @@ void FAdvancedRenamerModule::ShutdownModule()
 	FAdvancedRenamerStyle::Shutdown();
 	FAdvancedRenamerContentBrowserIntegration::Shutdown();
 	FAdvancedRenamerLevelEditorIntegration::Shutdown();
-	FAdvancedRenamerLevelEditorIntegration::ShutdownMenu();
-	if (EnableRenamerHandle.IsValid())
-	{
-		using namespace UE::AdvancedRenamer::Private;
-		CVarBatchRenamerEnableIntegration.AsVariable()->OnChangedDelegate().Remove(EnableRenamerHandle);
-	}
 }
 
 TSharedRef<IAdvancedRenamer> FAdvancedRenamerModule::CreateAdvancedRenamer(const TSharedRef<IAdvancedRenamerProvider>& InRenameProvider)
@@ -252,15 +205,6 @@ TArray<AActor*> FAdvancedRenamerModule::GetActorsSharingClassesInWorld(const TAr
 	return AllActors;
 }
 
-void FAdvancedRenamerModule::RegisterDefaultSections()
-{
-	Sections.Add(IAdvancedRenamerSection::MakeInstance<FAdvancedRenamerSearchAndReplaceSection>());
-	Sections.Add(IAdvancedRenamerSection::MakeInstance<FAdvancedRenamerRemovePrefixSection>());
-	Sections.Add(IAdvancedRenamerSection::MakeInstance<FAdvancedRenamerRemoveSuffixSection>());
-	Sections.Add(IAdvancedRenamerSection::MakeInstance<FAdvancedRenamerAddPrefixSuffixSection>());
-	Sections.Add(IAdvancedRenamerSection::MakeInstance<FAdvancedRenamerChangeCaseSection>());
-}
-
 IMPLEMENT_MODULE(FAdvancedRenamerModule, AdvancedRenamer)
 
 #undef LOCTEXT_NAMESPACE
diff --git a/Engine/Plugins/Experimental/AdvancedRenamer/Source/AdvancedRenamer/Private/AdvancedRenamerModule.h b/Engine/Plugins/Experimental/AdvancedRenamer/Source/AdvancedRenamer/Private/AdvancedRenamerModule.h
--- a/Engine/Plugins/Experimental/AdvancedRenamer/Source/AdvancedRenamer/Private/AdvancedRenamerModule.h
+++ b/Engine/Plugins/Experimental/AdvancedRenamer/Source/AdvancedRenamer/Private/AdvancedRenamerModule.h
@@ -3,8 +3,6 @@
 #pragma once
 
 #include "IAdvancedRenamerModule.h"
-#include "AdvancedRenamerSections/IAdvancedRenamerSection.h"
-#include "Delegates/IDelegateInstance.h"
 #include "Logging/LogMacros.h"
 
 DECLARE_LOG_CATEGORY_EXTERN(LogARP, Log, All);
@@ -26,13 +24,5 @@ class FAdvancedRenamerModule : public IAdvancedRenamerModule
 	virtual void OpenAdvancedRenamer(const TSharedRef<IAdvancedRenamer>& InRenamer, const TSharedPtr<IToolkitHost>& InToolkitHost) override;
 	virtual void OpenAdvancedRenamer(const TSharedRef<IAdvancedRenamer>& InRenamer, const TSharedPtr<SWidget>& InParentWidget) override;
 	virtual TArray<AActor*> GetActorsSharingClassesInWorld(const TArray<AActor*>& InActors) override;
-	virtual TArray<TSharedPtr<IAdvancedRenamerSection>> GetRegisteredSections() const override { return Sections; }
 	//~ End IAdvancedRenamerModule
-
-private:
-	void RegisterDefaultSections();
-
-private:
-	TArray<TSharedPtr<IAdvancedRenamerSection>> Sections;
-	FDelegateHandle EnableRenamerHandle;
 };
diff --git a/Engine/Plugins/Experimental/AdvancedRenamer/Source/AdvancedRenamer/Private/AdvancedRenamerSections/AdvancedRenamerAddPrefixSuffixSection.cpp b/Engine/Plugins/Experimental/AdvancedRenamer/Source/AdvancedRenamer/Private/AdvancedRenamerSections/AdvancedRenamerAddPrefixSuffixSection.cpp
deleted file mode 100644
diff --git a/Engine/Plugins/Experimental/AdvancedRenamer/Source/AdvancedRenamer/Private/AdvancedRenamerSections/AdvancedRenamerAddPrefixSuffixSection.h b/Engine/Plugins/Experimental/AdvancedRenamer/Source/AdvancedRenamer/Private/AdvancedRenamerSections/AdvancedRenamerAddPrefixSuffixSection.h
deleted file mode 100644
diff --git a/Engine/Plugins/Experimental/AdvancedRenamer/Source/AdvancedRenamer/Private/AdvancedRenamerSections/AdvancedRenamerChangeCaseSection.cpp b/Engine/Plugins/Experimental/AdvancedRenamer/Source/AdvancedRenamer/Private/AdvancedRenamerSections/AdvancedRenamerChangeCaseSection.cpp
deleted file mode 100644
diff --git a/Engine/Plugins/Experimental/AdvancedRenamer/Source/AdvancedRenamer/Private/AdvancedRenamerSections/AdvancedRenamerChangeCaseSection.h b/Engine/Plugins/Experimental/AdvancedRenamer/Source/AdvancedRenamer/Private/AdvancedRenamerSections/AdvancedRenamerChangeCaseSection.h
deleted file mode 100644
diff --git a/Engine/Plugins/Experimental/AdvancedRenamer/Source/AdvancedRenamer/Private/AdvancedRenamerSections/AdvancedRenamerRemovePrefixSection.cpp b/Engine/Plugins/Experimental/AdvancedRenamer/Source/AdvancedRenamer/Private/AdvancedRenamerSections/AdvancedRenamerRemovePrefixSection.cpp
deleted file mode 100644
diff --git a/Engine/Plugins/Experimental/AdvancedRenamer/Source/AdvancedRenamer/Private/AdvancedRenamerSections/AdvancedRenamerRemovePrefixSection.h b/Engine/Plugins/Experimental/AdvancedRenamer/Source/AdvancedRenamer/Private/AdvancedRenamerSections/AdvancedRenamerRemovePrefixSection.h
deleted file mode 100644
diff --git a/Engine/Plugins/Experimental/AdvancedRenamer/Source/AdvancedRenamer/Private/AdvancedRenamerSections/AdvancedRenamerRemoveSuffixSection.cpp b/Engine/Plugins/Experimental/AdvancedRenamer/Source/AdvancedRenamer/Private/AdvancedRenamerSections/AdvancedRenamerRemoveSuffixSection.cpp
deleted file mode 100644
diff --git a/Engine/Plugins/Experimental/AdvancedRenamer/Source/AdvancedRenamer/Private/AdvancedRenamerSections/AdvancedRenamerRemoveSuffixSection.h b/Engine/Plugins/Experimental/AdvancedRenamer/Source/AdvancedRenamer/Private/AdvancedRenamerSections/AdvancedRenamerRemoveSuffixSection.h
deleted file mode 100644
diff --git a/Engine/Plugins/Experimental/AdvancedRenamer/Source/AdvancedRenamer/Private/AdvancedRenamerSections/AdvancedRenamerSearchAndReplaceSection.cpp b/Engine/Plugins/Experimental/AdvancedRenamer/Source/AdvancedRenamer/Private/AdvancedRenamerSections/AdvancedRenamerSearchAndReplaceSection.cpp
deleted file mode 100644
diff --git a/Engine/Plugins/Experimental/AdvancedRenamer/Source/AdvancedRenamer/Private/AdvancedRenamerSections/AdvancedRenamerSearchAndReplaceSection.h b/Engine/Plugins/Experimental/AdvancedRenamer/Source/AdvancedRenamer/Private/AdvancedRenamerSections/AdvancedRenamerSearchAndReplaceSection.h
deleted file mode 100644
diff --git a/Engine/Plugins/Experimental/AdvancedRenamer/Source/AdvancedRenamer/Private/AdvancedRenamerSections/AdvancedRenamerSectionBase.cpp b/Engine/Plugins/Experimental/AdvancedRenamer/Source/AdvancedRenamer/Private/AdvancedRenamerSections/AdvancedRenamerSectionBase.cpp
deleted file mode 100644
diff --git a/Engine/Plugins/Experimental/AdvancedRenamer/Source/AdvancedRenamer/Private/AdvancedRenamerSections/AdvancedRenamerSectionBase.h b/Engine/Plugins/Experimental/AdvancedRenamer/Source/AdvancedRenamer/Private/AdvancedRenamerSections/AdvancedRenamerSectionBase.h
deleted file mode 100644
diff --git a/Engine/Plugins/Experimental/AdvancedRenamer/Source/AdvancedRenamer/Private/AdvancedRenamerSections/IAdvancedRenamerSection.h b/Engine/Plugins/Experimental/AdvancedRenamer/Source/AdvancedRenamer/Private/AdvancedRenamerSections/IAdvancedRenamerSection.h
deleted file mode 100644
diff --git a/Engine/Plugins/Experimental/AdvancedRenamer/Source/AdvancedRenamer/Private/AdvancedRenamerStyle.cpp b/Engine/Plugins/Experimental/AdvancedRenamer/Source/AdvancedRenamer/Private/AdvancedRenamerStyle.cpp
--- a/Engine/Plugins/Experimental/AdvancedRenamer/Source/AdvancedRenamer/Private/AdvancedRenamerStyle.cpp
+++ b/Engine/Plugins/Experimental/AdvancedRenamer/Source/AdvancedRenamer/Private/AdvancedRenamerStyle.cpp
@@ -1,8 +1,11 @@
 // Copyright Epic Games, Inc. All Rights Reserved.
 
 #include "AdvancedRenamerStyle.h"
+#include "Brushes/SlateImageBrush.h"
 #include "Interfaces/IPluginManager.h"
+#include "Misc/Paths.h"
 #include "Styling/AppStyle.h"
+#include "Styling/ISlateStyle.h"
 #include "Styling/SlateStyle.h"
 #include "Styling/SlateStyleRegistry.h"
 #include "Styling/SlateTypes.h"
@@ -14,7 +17,7 @@ void FAdvancedRenamerStyle::Initialize()
 {
 	if (!StyleInstance.IsValid())
 	{
-		InitStyle();
+		StyleInstance = Create();
 		FSlateStyleRegistry::RegisterSlateStyle(*StyleInstance);
 	}
 }
@@ -35,44 +38,37 @@ FName FAdvancedRenamerStyle::GetStyleSetName()
 	return StyleSetName;
 }
 
-void FAdvancedRenamerStyle::InitStyle()
+TSharedRef<FSlateStyleSet> FAdvancedRenamerStyle::Create()
 {
-	if (StyleInstance.IsValid())
-	{
-		return;
-	}
-
-	StyleInstance = MakeShared<FSlateStyleSet>("AdvancedRenamer");
+	TSharedRef<FSlateStyleSet> Style = MakeShared<FSlateStyleSet>("AdvancedRenamer");
 
 	TSharedPtr<IPlugin> Plugin = IPluginManager::Get().FindPlugin(TEXT("AdvancedRenamer"));
 	check(Plugin.IsValid());
 
 	if (Plugin.IsValid())
 	{
-		StyleInstance->SetContentRoot(FPaths::Combine(Plugin->GetBaseDir(), TEXT("Resources")));
+		Style->SetContentRoot(FPaths::Combine(Plugin->GetBaseDir(), TEXT("Resources")));
 	}
 
-	const FSplitterStyle SplitterStyle = FSplitterStyle()
-		.SetHandleNormalBrush(FSlateNoResource())
-		.SetHandleHighlightBrush(FSlateNoResource());
-	
-	StyleInstance->Set("AdvancedRenamer.Style.Splitter", SplitterStyle);
-
-	FSlateBrush* BackgroundBorderBrush = new FSlateColorBrush(FLinearColor::FromSRGBColor(FColor(36, 36, 36)));
+	Style->Set("AdvancedRenamer.Image.Radio.BlackBackground", new FSlateVectorImageBrush(
+		FPaths::Combine(Style->GetContentRootDir(), "/Icons/radio-background.svg"), 
+		FVector2D(16.f, 16.f), FStyleColors::Foldout.GetSpecifiedColor()
+	));
 
-	StyleInstance->Set("AdvancedRenamer.Style.BackgroundBorder", BackgroundBorderBrush);
+	FCheckBoxStyle BlackRadioButton = FAppStyle::GetWidgetStyle<FCheckBoxStyle>("RadioButton");
+	BlackRadioButton.SetBackgroundImage(*Style->GetBrush("AdvancedRenamer.Image.Radio.BlackBackground"));
 
-	const FTableViewStyle ListViewStyle = FTableViewStyle()
-		.SetBackgroundBrush(*FAdvancedRenamerStyle::Get().GetBrush("AdvancedRenamer.Style.BackgroundBorder"));
+	Style->Set("AdvancedRenamer.Style.BlackRadioButton", BlackRadioButton);
 
-	StyleInstance->Set("AdvancedRenamer.Style.ListView", ListViewStyle);
+	const FButtonStyle DarkButton = FButtonStyle(FAppStyle::Get().GetWidgetStyle<FButtonStyle>("SimpleButton"))
+		.SetNormal(FSlateColorBrush(FStyleColors::Recessed))
+		.SetHovered(FSlateColorBrush(FStyleColors::Hover))
+		.SetPressed(FSlateColorBrush(FStyleColors::Header))
+		.SetDisabled(FSlateColorBrush(FStyleColors::Dropdown));
 
-	FHeaderRowStyle HeaderRowStyle = FAppStyle::Get().GetWidgetStyle<FHeaderRowStyle>("TableView.Header");
-	HeaderRowStyle.SetHorizontalSeparatorThickness(0);
-	HeaderRowStyle.SetHorizontalSeparatorBrush(FSlateNoResource());
-	HeaderRowStyle.SetBackgroundBrush(FSlateColorBrush(FLinearColor::FromSRGBColor(FColor(47, 47, 47))));
+	Style->Set("AdvancedRenamer.Style.DarkButton", DarkButton);
 
-	StyleInstance->Set("AdvancedRenamer.Style.HeaderRow", HeaderRowStyle);
+	return Style;
 }
 
 const ISlateStyle& FAdvancedRenamerStyle::Get()
diff --git a/Engine/Plugins/Experimental/AdvancedRenamer/Source/AdvancedRenamer/Private/AdvancedRenamerStyle.h b/Engine/Plugins/Experimental/AdvancedRenamer/Source/AdvancedRenamer/Private/AdvancedRenamerStyle.h
--- a/Engine/Plugins/Experimental/AdvancedRenamer/Source/AdvancedRenamer/Private/AdvancedRenamerStyle.h
+++ b/Engine/Plugins/Experimental/AdvancedRenamer/Source/AdvancedRenamer/Private/AdvancedRenamerStyle.h
@@ -22,7 +22,7 @@ public:
 
 private:
 
-	static void InitStyle();
+	static TSharedRef<FSlateStyleSet> Create();
 	static TSharedPtr<FSlateStyleSet> StyleInstance;
 
 };
diff --git a/Engine/Plugins/Experimental/AdvancedRenamer/Source/AdvancedRenamer/Private/IAdvancedRenamer.h b/Engine/Plugins/Experimental/AdvancedRenamer/Source/AdvancedRenamer/Private/IAdvancedRenamer.h
--- a/Engine/Plugins/Experimental/AdvancedRenamer/Source/AdvancedRenamer/Private/IAdvancedRenamer.h
+++ b/Engine/Plugins/Experimental/AdvancedRenamer/Source/AdvancedRenamer/Private/IAdvancedRenamer.h
@@ -2,17 +2,105 @@
 
 #pragma once
 
-#include "AdvancedRenamerExecuteSection.h"
+#include "Containers/ContainersFwd.h"
 #include "Containers/UnrealString.h"
+#include "HAL/Platform.h"
+#include "Misc/CString.h"
 #include "Providers/IAdvancedRenamerProvider.h"
 #include "Templates/SharedPointerFwd.h"
 
+enum class EAdvancedRenamerSeachAndReplaceType : uint8
+{
+	None,
+	PlainText,
+	RegularExpression
+};
+
+/**
+ * Options for the rename engine.
+ *
+ * - RemovePrefixSeparator and RemovePrefixCharacterCount are mutually exclusive.
+ * - RemoveSuffixSeparator and RemoveSuffixCharacterCount are mutually exclusive.
+ * - bAddSuffixNumber implies bRemoveSuffixNumer.
+ */
+struct FAdvancedRenamerOptions
+{
+	/** Name => OtherName */
+	FString BaseName = ""; 
+
+	/** Name => Other_Name */
+	FString AddPrefix = ""; 
+
+	/** Other_Name, _ => Name */
+	FString RemovePrefixSeparator = "";
+
+	/** OtherName, 5 => Name */
+	int32 RemovePrefixCharacterCount = 0; 
+
+	/** Name => Name_Other */
+	FString AddSuffix = ""; 
+
+	/** Name => Name1 */
+	bool bAddSuffixNumber = false; 
+
+	/** Name, 5 => Name5 */
+	int32 AddSuffixNumberStart = 1; 
+
+	/** Name, Moo, 1, 5 => Name1, Moo6 */
+	int32 AddSuffixNumberStep = 1;
+
+	/** Name_Other, _ => Name */
+	FString RemoveSuffixSeparator = ""; 
+
+	/** NameOther, 5 => Name */
+	int32 RemoveSuffixCharacterCount = 0; 
+
+	/** Name5 => Name */
+	bool bRemoveSuffixNumber = false; 
+
+	EAdvancedRenamerSeachAndReplaceType SearchAndReplaceType = EAdvancedRenamerSeachAndReplaceType::None;
+	ESearchCase::Type SearchAndReplaceCase = ESearchCase::CaseSensitive;
+
+	/** Name, am, ot => Note */
+	FString SearchAndReplaceFromText = "";
+
+	/** Name, am, ot => Note */
+	FString SearchAndReplaceToText = "";
+
+	bool HasBaseName() const
+	{
+		return !BaseName.IsEmpty();
+	}
+
+	bool HasPrefix() const
+	{
+		return !AddPrefix.IsEmpty()
+			|| !RemovePrefixSeparator.IsEmpty()
+			|| RemovePrefixCharacterCount > 0;
+	}
+
+	bool HasSuffix() const
+	{
+		return !AddSuffix.IsEmpty()
+			|| (bAddSuffixNumber && AddSuffixNumberStep > 0)
+			|| !RemoveSuffixSeparator.IsEmpty()
+			|| RemoveSuffixCharacterCount > 0
+			|| bRemoveSuffixNumber;			
+	}
+
+	bool HasSearchAndReplace() const
+	{
+		return SearchAndReplaceType != EAdvancedRenamerSeachAndReplaceType::None
+			&& !SearchAndReplaceFromText.IsEmpty();
+	}
+};
+
 struct FAdvancedRenamerPreview
 {
 	FAdvancedRenamerPreview(int32 InHash, const FString InOriginalName)
 		: Hash(InHash)
 		, OriginalName(InOriginalName)
-		, NewName(TEXT(""))
+		, NewName(FString(""))
 	{
 	}
 
@@ -32,14 +120,18 @@ public:
 
 	virtual const TSharedRef<IAdvancedRenamerProvider>& GetProvider() const = 0;
 
+	virtual const FAdvancedRenamerOptions& GetOptions() const = 0;
+
+	/** Make sure to run UpdatePreviews or MarkDirty if you make changes directly. */
+	virtual FAdvancedRenamerOptions& GetOptions() = 0;
+
+	virtual void SetOption(const FAdvancedRenamerOptions& InOptions) = 0;
+
 	virtual const TArray<TSharedPtr<FAdvancedRenamerPreview>>& GetPreviews() = 0;
 
 	/** Returns the preview for the item at the given index. */
 	virtual TSharedPtr<FAdvancedRenamerPreview> GetPreview(int32 InIndex) const = 0;
 
-	/** Add a section to this Renamer */
-	virtual void AddSection(FAdvancedRenamerExecuteSection InSection) = 0;
-
 	/** True if there are any items actually renamed by the preview generator. */
 	virtual bool HasRenames() const = 0;
 
@@ -49,8 +141,8 @@ public:
 	virtual void MarkClean() = 0;
 
 	/** Executes the rename on the given name. */
-	virtual FString ApplyRename(const FString& InName) = 0;
-	
+	virtual FString ApplyRename(const FString& InName, int32 InIndex) const = 0;
+
 	/** Returns true if any names actually changed. */
 	virtual bool UpdatePreviews() = 0;
 
diff --git a/Engine/Plugins/Experimental/AdvancedRenamer/Source/AdvancedRenamer/Private/Integrations/AdvancedRenamerContentBrowserIntegration.cpp b/Engine/Plugins/Experimental/AdvancedRenamer/Source/AdvancedRenamer/Private/Integrations/AdvancedRenamerContentBrowserIntegration.cpp
--- a/Engine/Plugins/Experimental/AdvancedRenamer/Source/AdvancedRenamer/Private/Integrations/AdvancedRenamerContentBrowserIntegration.cpp
+++ b/Engine/Plugins/Experimental/AdvancedRenamer/Source/AdvancedRenamer/Private/Integrations/AdvancedRenamerContentBrowserIntegration.cpp
@@ -29,8 +29,8 @@ namespace UE::AdvancedRenamer::Private
 	void ExtendAssetMenu(FMenuBuilder& MenuBuilder, const TArray<FAssetData> AssetArray)
 	{
 		MenuBuilder.AddMenuEntry(
-			LOCTEXT("AdvancedRename", "Batch Rename"),
-			LOCTEXT("AdvancedRenameTooltip", "Opens the Batch Renamer Panel to rename all selected assets."),
+			LOCTEXT("AdvancedRename", "Advanced Rename"),
+			LOCTEXT("AdvancedRenameTooltip", "Opens the Advanced Renamer Panel to rename all selected assets."),
 			FSlateIcon(FAppStyle::GetAppStyleSetName(), "ContentBrowser.AssetActions.Rename"),
 			FUIAction(FExecuteAction::CreateStatic(&OpenAdvancedRenamer, AssetArray))
 		);
diff --git a/Engine/Plugins/Experimental/AdvancedRenamer/Source/AdvancedRenamer/Private/Integrations/AdvancedRenamerLevelEditorIntegration.cpp b/Engine/Plugins/Experimental/AdvancedRenamer/Source/AdvancedRenamer/Private/Integrations/AdvancedRenamerLevelEditorIntegration.cpp
--- a/Engine/Plugins/Experimental/AdvancedRenamer/Source/AdvancedRenamer/Private/Integrations/AdvancedRenamerLevelEditorIntegration.cpp
+++ b/Engine/Plugins/Experimental/AdvancedRenamer/Source/AdvancedRenamer/Private/Integrations/AdvancedRenamerLevelEditorIntegration.cpp
@@ -102,23 +102,20 @@ namespace UE::AdvancedRenamer::Private
 	}
 }
 
-static const TArray<FName, TInlineAllocator<2>> Menus = {
-	"LevelEditor.ActorContextMenu.EditSubMenu",
-	"LevelEditor.LevelEditorSceneOutliner.ContextMenu.EditSubMenu"
-};
-
 void FAdvancedRenamerLevelEditorIntegration::Initialize()
 {
 	using namespace UE::AdvancedRenamer::Private;
 
 	FLevelEditorModule& LevelEditorModule = FModuleManager::LoadModuleChecked<FLevelEditorModule>("LevelEditor");
 	LevelEditorCreatedDelegateHandle = LevelEditorModule.OnLevelEditorCreated().AddStatic(&OnLevelEditorCreated);
-}
 
-void FAdvancedRenamerLevelEditorIntegration::InitializeMenu()
-{
 	const FAdvancedRenamerCommands& AdvRenCommands = FAdvancedRenamerCommands::Get();
 
+	static const TArray<FName, TInlineAllocator<2>> Menus = {
+		"LevelEditor.ActorContextMenu.EditSubMenu",
+		"LevelEditor.LevelEditorSceneOutliner.ContextMenu.EditSubMenu"
+	};
+
 	const TAttribute<FText> TextAttribute;
 	const FSlateIcon RenameIcon = FSlateIcon(FAppStyle::GetAppStyleSetName(), "GenericCommands.Rename");
 
@@ -150,18 +147,5 @@ void FAdvancedRenamerLevelEditorIntegration::Shutdown()
 	if (FLevelEditorModule* LevelEditorModule = FModuleManager::LoadModulePtr<FLevelEditorModule>("LevelEditor"))
 	{
 		LevelEditorModule->OnLevelEditorCreated().Remove(LevelEditorCreatedDelegateHandle);
-		LevelEditorCreatedDelegateHandle.Reset();
-	}
-}
-
-void FAdvancedRenamerLevelEditorIntegration::ShutdownMenu()
-{
-	if (UToolMenus* ToolMenus = UToolMenus::Get())
-	{
-		for (const FName& Menu : Menus)
-		{
-			ToolMenus->RemoveEntry(Menu, NAME_None, TEXT("RenameSelectedActors"));
-			ToolMenus->RemoveEntry(Menu, NAME_None, TEXT("RenameSharedClassActors"));
-		}
 	}
 }
diff --git a/Engine/Plugins/Experimental/AdvancedRenamer/Source/AdvancedRenamer/Private/Integrations/AdvancedRenamerLevelEditorIntegration.h b/Engine/Plugins/Experimental/AdvancedRenamer/Source/AdvancedRenamer/Private/Integrations/AdvancedRenamerLevelEditorIntegration.h
--- a/Engine/Plugins/Experimental/AdvancedRenamer/Source/AdvancedRenamer/Private/Integrations/AdvancedRenamerLevelEditorIntegration.h
+++ b/Engine/Plugins/Experimental/AdvancedRenamer/Source/AdvancedRenamer/Private/Integrations/AdvancedRenamerLevelEditorIntegration.h
@@ -6,7 +6,5 @@ class FAdvancedRenamerLevelEditorIntegration
 {
 public:
 	static void Initialize();
-	static void InitializeMenu();
 	static void Shutdown();	
-	static void ShutdownMenu();
 };
diff --git a/Engine/Plugins/Experimental/AdvancedRenamer/Source/AdvancedRenamer/Private/Slate/SAdvancedRenamerPanel.cpp b/Engine/Plugins/Experimental/AdvancedRenamer/Source/AdvancedRenamer/Private/Slate/SAdvancedRenamerPanel.cpp
--- a/Engine/Plugins/Experimental/AdvancedRenamer/Source/AdvancedRenamer/Private/Slate/SAdvancedRenamerPanel.cpp
+++ b/Engine/Plugins/Experimental/AdvancedRenamer/Source/AdvancedRenamer/Private/Slate/SAdvancedRenamerPanel.cpp
@@ -1,31 +1,60 @@
 // Copyright Epic Games, Inc. All Rights Reserved.
 
 #include "Slate/SAdvancedRenamerPanel.h"
-
-#include "AdvancedRenamerModule.h"
-#include "AdvancedRenamerSections/IAdvancedRenamerSection.h"
 #include "AdvancedRenamerStyle.h"
 #include "EngineAnalytics.h"
-#include "Framework/Application/SlateApplication.h"
 #include "Framework/Commands/GenericCommands.h"
 #include "Framework/MultiBox/MultiBoxBuilder.h"
-#include "GenericPlatform/GenericPlatformMath.h"
 #include "IAdvancedRenamer.h"
 #include "ScopedTransaction.h"
-#include "Framework/Notifications/NotificationManager.h"
-#include "Utils/AdvancedRenamerSlateUtils.h"
+#include "Styling/AppStyle.h"
+#include "Styling/StyleColors.h"
+#include "Widgets/Colors/SColorBlock.h"
 #include "Widgets/Input/SButton.h"
+#include "Widgets/Input/SCheckBox.h"
+#include "Widgets/Input/SEditableTextBox.h"
+#include "Widgets/Input/SMultiLineEditableTextBox.h"
+#include "Widgets/Input/SSpinBox.h"
 #include "Widgets/Layout/SBorder.h"
-#include "Widgets/Layout/SSpacer.h"
-#include "Widgets/Layout/SSplitter.h"
-#include "Widgets/SBoxPanel.h"
-#include "Widgets/SWindow.h"
-#include "Widgets/Notifications/SNotificationList.h"
-#include "Widgets/Text/STextBlock.h"
+#include "Widgets/Layout/SBox.h"
+#include "Widgets/Layout/SScrollBox.h"
+#include "Widgets/SCanvas.h"
+#include "Widgets/SNullWidget.h"
 #include "Widgets/Views/SListView.h"
+#include "Widgets/Views/STableRow.h"
 
 #define LOCTEXT_NAMESPACE "SAdvancedRenamerPanel"
 
+namespace UE::AdvancedRenamer::Private
+{
+	const FVector2D WindowSize = {600.f, 500.f};
+	const FSlateFontInfo TitleFont = FCoreStyle::GetDefaultFontStyle("Regular", 12);
+	const FSlateFontInfo RegularFont = FCoreStyle::GetDefaultFontStyle("Regular", 10);
+	constexpr float LeftBlockWidth = 304.f;
+	const FVector2D TitleSize = {150.f, 20.f};
+	constexpr float TitleOffsetY = 25.f;
+	const FVector2D CheckboxOffset = {0.f, 2.f};
+	const FVector2D CheckboxSize = {18.f, 18.f};
+	const FVector2D RadioOffset = {1.f, 2.f};
+	constexpr float LabelStart = 25.f;
+	constexpr float LabelOffsetY = 3.f;
+	const FVector2D LabelSize = {125.f, 18.f};
+	constexpr float EntryStart = 145.f;
+	const FVector2D EntrySize = {155.f, 21.f};
+	constexpr float LineHeight = 26.f;
+	const FVector2D SeparatorSize = {32.f, 21.f};
+	const FVector2D SpinSize1 = {32.f, 21.f};
+	const FVector2D SpinSize2 = {39.f, 21.f};
+	const FVector2D SpinSize3 = {46.f, 21.f};
+	constexpr float RightBlockOffsetX = 314.f;
+	const FVector2D RightBlockSize = {281.f, 490.f};
+	const FVector2D ListViewSize = {277.f, 431.f};
+	constexpr float ListLineHeight = 15.f;
+	constexpr float ApplyButtonHeight = 25.f;
+	static FName OriginalNameColumnName = "OriginalName";
+	static FName NewNameColumnName = "NewName";
+}
+
 class SAdvancedRenamerPreviewListRow : public SMultiColumnTableRow<TSharedPtr<FAdvancedRenamerPreview>>
 {
 public:
@@ -39,12 +68,12 @@ public:
 		ItemWeak = InRowItem;
 
 		SMultiColumnTableRow<TSharedPtr<FAdvancedRenamerPreview>>::Construct(FSuperRowType::FArguments(), InOwnerTableView);
-		SetBorderImage(FAdvancedRenamerStyle::Get().GetBrush("AdvancedRenamer.Style.BackgroundBorder"));
+		SetBorderImage(TAttribute<const FSlateBrush*>(this, &SAdvancedRenamerPreviewListRow::GetBorder));
 	}
 
 	virtual TSharedRef<SWidget> GenerateWidgetForColumn(const FName& InColumnName) override
 	{
-		using namespace AdvancedRenamerSlateUtils::Default;
+		using namespace UE::AdvancedRenamer::Private;
 
 		if (InColumnName != OriginalNameColumnName
 			&& InColumnName != NewNameColumnName)
@@ -67,22 +96,22 @@ public:
 		}
 
 		TSharedRef<STextBlock> TextBlock = SNew(STextBlock)
-			.Font(RegularFont);
-
-		TSharedRef<SBox> TextBlockBox = SNew(SBox)
-			.Padding(FMargin(8.f, 6.f, 0.f, 6.f));
+			.ColorAndOpacity(FStyleColors::AccentWhite.GetSpecifiedColor())
+			.Font(RegularFont)
+			.Margin(FMargin(5.f, 0.f, 5.f, 0.f));
 
 		if (InColumnName == OriginalNameColumnName)
 		{
 			TextBlock->SetText(FText::FromString(Item->OriginalName));
+			Panel->MinDesiredOriginalNameWidth = FMath::Max(Panel->MinDesiredOriginalNameWidth, TextBlock->ComputeDesiredSize(1.f).X);
 		}
 		else if (InColumnName == NewNameColumnName)
 		{
 			TextBlock->SetText(FText::FromString(Item->NewName));
+			Panel->MinDesiredNewNameWidth = FMath::Max(Panel->MinDesiredOriginalNameWidth, TextBlock->ComputeDesiredSize(1.f).X);
 		}
 
-		TextBlockBox->SetContent(TextBlock);
-		return TextBlockBox;
+		return TextBlock;
 	}
 
 protected:
@@ -94,6 +123,11 @@ void SAdvancedRenamerPanel::Construct(const FArguments& InArgs, const TSharedRef
 {
 	Renamer = InRenamer;
 
+	bRemovePrefixSeparator = false;
+	bRemovePrefixNumChars = false;
+	bRemoveSuffixSeparator = false;
+	bRemoveSuffixNumChars = false;
+
 	CommandList = MakeShared<FUICommandList>();
 	CommandList->MapAction(
 		FGenericCommands::Get().Delete,
@@ -101,251 +135,637 @@ void SAdvancedRenamerPanel::Construct(const FArguments& InArgs, const TSharedRef
 		FCanExecuteAction()
 	);
 
-	using namespace AdvancedRenamerSlateUtils::Default;
+	using namespace UE::AdvancedRenamer::Private;
+
+	TSharedRef<SCanvas> Canvas = SNew(SCanvas)
+		+ SCanvas::Slot()
+		.Position(FVector2D(0.f, 0.f))
+		.Size(WindowSize)
+		[
+			SNew(SColorBlock)
+			.Color(FStyleColors::Background.GetSpecifiedColor())
+		];
+
+	CreateLeftPane(Canvas);
+	CreateRightPane(Canvas);
 
 	ChildSlot
 	[
-		SNew(SVerticalBox)
+		Canvas
+	];
+
+	if (FEngineAnalytics::IsAvailable())
+	{
+		FEngineAnalytics::GetProvider().RecordEvent(TEXT("Editor.Usage.AdvancedRenamer.Opened"));
+	}
+}
 
-		+ SVerticalBox::Slot()
-		.FillHeight(1.f)
+void SAdvancedRenamerPanel::CreateLeftPane(const TSharedRef<SCanvas>& InCanvas)
+{	
+	using namespace UE::AdvancedRenamer::Private;
+
+	InCanvas->AddSlot()
+		.Position(FVector2D(5.f, 5.f))
+		.Size(FVector2D(LeftBlockWidth, 50.f))
 		[
-			SNew(SBorder)
-			.BorderImage(FAppStyle::GetBrush("NoBorder"))
-			.Padding(2.f, 2.f)
-			[
-				SNew(SBox)
-				[
-					SNew(SSplitter)
-					.PhysicalSplitterHandleSize(2.f)
-					.MinimumSlotHeight(499.5f)
-					.Style(&FAdvancedRenamerStyle::Get().GetWidgetStyle<FSplitterStyle>("AdvancedRenamer.Style.Splitter"))
-
-					+ SSplitter::Slot()
-					.Value(0.505f)
-					.MinSize(365.f)
-					[
-						SAssignNew(LeftSideVerticalBox, SVerticalBox)
-					]
-
-					+ SSplitter::Slot()
-					.Value(0.495f)
-					[
-						SAssignNew(RightSideBox, SBox)
-					]
-				]
-			]
-		]
+			CreateBaseName()
+		];
+
+	InCanvas->AddSlot()
+		.Position(FVector2D(5.f, 60.f))
+		.Size(FVector2D(LeftBlockWidth, 102.f))
+		[
+			CreatePrefix()
+		];
 
-		+ SVerticalBox::Slot()
-		.AutoHeight()
+	InCanvas->AddSlot()
+		.Position(FVector2D(5.f, 167.f))
+		.Size(FVector2D(LeftBlockWidth, 154.f))
 		[
-			SNew(SBorder)
-			.BorderImage(FAdvancedRenamerStyle::Get().GetBrush("AdvancedRenamer.Style.BackgroundBorder"))
+			CreateSuffix()
+		];
+
+	InCanvas->AddSlot()
+		.Position(FVector2D(5.f, 327.f))
+		.Size(FVector2D(LeftBlockWidth, 168.f))
+		[
+			CreateSearchAndReplace()
+		];
+}
+
+TSharedRef<SWidget> SAdvancedRenamerPanel::CreateBaseName()
+{
+	using namespace UE::AdvancedRenamer::Private;
+
+	return SNew(SBorder)
+		.BorderImage(FAppStyle::Get().GetBrush("DetailsView.CategoryTop"))
+		.Content()
+		[
+			SNew(SCanvas)
+			+ SCanvas::Slot()
+			.Position(FVector2D(0.f, LabelOffsetY))
+			.Size(TitleSize)
 			[
-				SAssignNew(ApplyResetCancelBox, SBox)
+				SNew(STextBlock)
+				.ColorAndOpacity(FStyleColors::AccentBlue.GetSpecifiedColor())
+				.Font(TitleFont)
+				.Text(LOCTEXT("BaseNameTitle", "Name"))
 			]
-		]
-	];
-
-	TArray<TSharedPtr<IAdvancedRenamerSection>> RegisteredSections = IAdvancedRenamerModule::Get().GetRegisteredSections();
-	for (int32 Index = 0; Index < RegisteredSections.Num(); Index++)
-	{
-		RegisteredSections[Index]->Init(InRenamer);
-		FMargin Padding = FMargin(0.f, 1.f);
+			+ SCanvas::Slot()
+			.Position(FVector2D(LabelStart * 2.f, TitleOffsetY + LabelOffsetY))
+			.Size(LabelSize)
+			[
+				SNew(STextBlock)
+				.ColorAndOpacity(FStyleColors::AccentWhite.GetSpecifiedColor())
+				.Font(RegularFont)
+				.Text(LOCTEXT("SetBaseName", "Set Base Name"))
+			]
+			+ SCanvas::Slot()
+			.Position(FVector2D(EntryStart, TitleOffsetY))
+			.Size(EntrySize)
+			[
+				SAssignNew(BaseNameTextBox, SEditableTextBox)
+				.BackgroundColor(FStyleColors::Background.GetSpecifiedColor())
+				.ForegroundColor(FStyleColors::AccentWhite.GetSpecifiedColor())
+				.Font(RegularFont)
+				.HintText(LOCTEXT("BaseNameHint", "Base name"))
+				.OnTextChanged(this, &SAdvancedRenamerPanel::OnBaseNameChanged)
+			]
+		];
+}
 
-		if (Index == 0)
-		{
-			// The first Section Search and Replace will fill the height by design
-			Padding = FMargin(0.f, 0.f, 0.f, 1.f);
-			LeftSideVerticalBox->AddSlot()
-				.Padding(Padding)
-				.FillHeight(1.f)
-				[
-					RegisteredSections[Index]->GetWidget()
-				];
-			continue;
-		}
+TSharedRef<SWidget> SAdvancedRenamerPanel::CreatePrefix()
+{
+	using namespace UE::AdvancedRenamer::Private;
 
-		if (Index == RegisteredSections.Num() - 1)
-		{
-			Padding = FMargin(0.f, 1.f, 0.f, 0.f);
-		}
+	return SNew(SBorder)
+		.BorderImage(FAppStyle::Get().GetBrush("DetailsView.CategoryTop"))
+		.Content()
+		[
+			SNew(SCanvas)
 
-		LeftSideVerticalBox->AddSlot()
-			.Padding(Padding)
-			.AutoHeight()
+			// Title
+			+ SCanvas::Slot()
+			.Position(FVector2D(0.f, LabelOffsetY))
+			.Size(TitleSize)
 			[
-				RegisteredSections[Index]->GetWidget()
-			];
-	}
+				SNew(STextBlock)
+				.ColorAndOpacity(FStyleColors::AccentBlue.GetSpecifiedColor())
+				.Font(TitleFont)
+				.Text(LOCTEXT("PrefixTitle", "Prefix"))
+			]
+			//
+			+ SCanvas::Slot()
+			.Position(FVector2D(LabelStart * 2.f, TitleOffsetY + LabelOffsetY))
+			.Size(LabelSize)
+			[
+				SNew(STextBlock)
+				.ColorAndOpacity(FStyleColors::AccentWhite.GetSpecifiedColor())
+				.Font(RegularFont)
+				.Text(LOCTEXT("AddPrefix", "Add Prefix"))
+			]
+			+ SCanvas::Slot()
+			.Position(FVector2D(EntryStart, TitleOffsetY))
+			.Size(EntrySize)
+			[
+				SAssignNew(PrefixTextBox, SEditableTextBox)
+				.BackgroundColor(FStyleColors::Background.GetSpecifiedColor())
+				.ForegroundColor(FStyleColors::AccentWhite.GetSpecifiedColor())
+				.Font(RegularFont)
+				.HintText(LOCTEXT("PrefixHint", "Prefix"))
+				.OnTextChanged(this, &SAdvancedRenamerPanel::OnPrefixChanged)
+			]
 
-	CreateRightPanel();
-	CreateApplyResetCancelPanel();
+			// Remove prefix
+			+ SCanvas::Slot()
+			.Position(FVector2D(LabelStart, TitleOffsetY + LineHeight) + RadioOffset)
+			.Size(CheckboxSize)
+			[
+				SAssignNew(PrefixRemoveCheckBox, SCheckBox)
+				.Style(&FAdvancedRenamerStyle::Get().GetWidgetStyle<FCheckBoxStyle>("AdvancedRenamer.Style.BlackRadioButton"))
+				.IsChecked(this, &SAdvancedRenamerPanel::IsPrefixRemoveChecked)
+				.OnCheckStateChanged(this, &SAdvancedRenamerPanel::OnPrefixRemoveCheckBoxChanged)
+			]
+			+ SCanvas::Slot()
+			.Position(FVector2D(LabelStart * 2.f, TitleOffsetY + LabelOffsetY + LineHeight))
+			.Size(LabelSize)
+			[
+				SNew(STextBlock)
+				.ColorAndOpacity(FStyleColors::AccentWhite.GetSpecifiedColor())
+				.Font(RegularFont)
+				.Text(LOCTEXT("PrefixRemove", "Remove Prefix"))
+			]
+			+ SCanvas::Slot()
+			.Position(FVector2D(EntryStart, TitleOffsetY + LineHeight))
+			.Size(SeparatorSize)
+			[
+				SAssignNew(PrefixSeparatorTextBox, SEditableTextBox)
+				.BackgroundColor(FStyleColors::Background.GetSpecifiedColor())
+				.ForegroundColor(FStyleColors::AccentWhite.GetSpecifiedColor())
+				.Font(RegularFont)
+				.IsEnabled(this, &SAdvancedRenamerPanel::IsPrefixRemoveSeparatorEnabled)
+				.Text(LOCTEXT("Underscore", "_"))
+				.OnVerifyTextChanged(this, &SAdvancedRenamerPanel::OnPrefixSeparatorVerifyTextChanged)
+				.OnTextChanged(this, &SAdvancedRenamerPanel::OnPrefixSeparatorChanged)
+			]
+			+ SCanvas::Slot()
+			.Position(FVector2D(EntryStart + SeparatorSize.X + 5.f, TitleOffsetY + LabelOffsetY + LineHeight))
+			.Size(LabelSize)
+			[
+				SNew(STextBlock)
+				.ColorAndOpacity(FStyleColors::AccentWhite.GetSpecifiedColor())
+				.Font(RegularFont)
+				.Text(LOCTEXT("Separator", "(separator)"))
+			]
 
-	Renamer->MarkDirty();
+			// Remove first
+			+ SCanvas::Slot()
+			.Position(FVector2D(LabelStart, TitleOffsetY + LineHeight * 2.f) + RadioOffset)
+			.Size(CheckboxSize)
+			[
+				SAssignNew(PrefixRemoveCharactersCheckBox, SCheckBox)
+				.Style(&FAdvancedRenamerStyle::Get().GetWidgetStyle<FCheckBoxStyle>("AdvancedRenamer.Style.BlackRadioButton"))
+				.IsChecked(this, &SAdvancedRenamerPanel::IsPrefixRemoveCharactersChecked)
+				.OnCheckStateChanged(this, &SAdvancedRenamerPanel::OnPrefixRemoveCharactersCheckBoxChanged)
+			]
+			+ SCanvas::Slot()
+			.Position(FVector2D(LabelStart * 2.f, TitleOffsetY + LabelOffsetY + LineHeight * 2.f))
+			.Size(LabelSize)
+			[
+				SNew(STextBlock)
+				.ColorAndOpacity(FStyleColors::AccentWhite.GetSpecifiedColor())
+				.Font(RegularFont)
+				.Text(LOCTEXT("PrefixRemoveChars", "Remove First"))
+			]
+			+ SCanvas::Slot()
+			.Position(FVector2D(EntryStart, TitleOffsetY + LineHeight * 2.f))
+			.Size(SpinSize1)
+			[
 
-	if (FEngineAnalytics::IsAvailable())
-	{
-		FEngineAnalytics::GetProvider().RecordEvent(TEXT("Editor.Usage.AdvancedRenamer.Opened"));
-	}
+				SAssignNew(PrefixRemoveCharactersSpinBox, SSpinBox<uint8>)
+				.Style(&FAppStyle::Get(), "Menu.SpinBox")
+				.Font(RegularFont)
+				.MinValue(1)
+				.MaxValue(9)
+				.Value(1)
+				.IsEnabled(this, &SAdvancedRenamerPanel::IsPrefixRemoveNumCharsEnabled)
+				.OnValueChanged(this, &SAdvancedRenamerPanel::OnPrefixRemoveCharactersChanged)
+			]
+			+ SCanvas::Slot()
+			.Position(FVector2D(EntryStart + SeparatorSize.X + 5.f, TitleOffsetY + LabelOffsetY + LineHeight * 2.f))
+			.Size(LabelSize)
+			[
+				SNew(STextBlock)
+				.ColorAndOpacity(FStyleColors::AccentWhite.GetSpecifiedColor())
+				.Font(RegularFont)
+				.Text(LOCTEXT("Characters", "character(s)"))
+			]
+		];
 }
 
-void SAdvancedRenamerPanel::CreateApplyResetCancelPanel()
+TSharedRef<SWidget> SAdvancedRenamerPanel::CreateSuffix()
 {
-	using namespace AdvancedRenamerSlateUtils::Default;
+	using namespace UE::AdvancedRenamer::Private;
 
-	ApplyResetCancelBox->SetContent(SNew(SHorizontalBox)
-			+ SHorizontalBox::Slot()
-			.HAlign(HAlign_Right)
+	return SNew(SBorder)
+		.BorderImage(FAppStyle::Get().GetBrush("DetailsView.CategoryTop"))
+		.Content()
+		[
+			SNew(SCanvas)
+
+			// Title
+			+ SCanvas::Slot()
+			.Position(FVector2D(0.f, LabelOffsetY))
+			.Size(TitleSize)
 			[
-				SNew(SSpacer)
+				SNew(STextBlock)
+				.ColorAndOpacity(FStyleColors::AccentBlue.GetSpecifiedColor())
+				.Font(TitleFont)
+				.Text(LOCTEXT("SuffixTitle", "Suffix"))
 			]
 
-			+ SHorizontalBox::Slot()
-			.Padding(ApplyButtonPadding)
-			.HAlign(HAlign_Right)
-			.AutoWidth()
+			// Suffix name
+			+ SCanvas::Slot()
+			.Position(FVector2D(LabelStart * 2.f, TitleOffsetY + LabelOffsetY))
+			.Size(LabelSize)
+			[
+				SNew(STextBlock)
+				.ColorAndOpacity(FStyleColors::AccentWhite.GetSpecifiedColor())
+				.Font(RegularFont)
+				.Text(LOCTEXT("AddSuffix", "Add Suffix"))
+			]
+			+ SCanvas::Slot()
+			.Position(FVector2D(EntryStart, TitleOffsetY))
+			.Size(EntrySize)
 			[
-				CreateApplyButton()
+				SAssignNew(SuffixTextBox, SEditableTextBox)
+				.BackgroundColor(FStyleColors::Background.GetSpecifiedColor())
+				.ForegroundColor(FStyleColors::AccentWhite.GetSpecifiedColor())
+				.Font(RegularFont)
+				.HintText(LOCTEXT("SuffixHint", "Suffix"))
+				.OnTextChanged(this, &SAdvancedRenamerPanel::OnSuffixChanged)
 			]
 
-			+ SHorizontalBox::Slot()
-			.Padding(ResetButtonPadding)
-			.HAlign(HAlign_Right)
-			.AutoWidth()
+			// Remove Suffix
+			+ SCanvas::Slot()
+			.Position(FVector2D(LabelStart, TitleOffsetY + LineHeight) + RadioOffset)
+			.Size(CheckboxSize)
 			[
-				CreateResetButton()
+				SAssignNew(SuffixRemoveCheckBox, SCheckBox)
+				.Style(&FAdvancedRenamerStyle::Get().GetWidgetStyle<FCheckBoxStyle>("AdvancedRenamer.Style.BlackRadioButton"))
+				.IsChecked(this, &SAdvancedRenamerPanel::IsSuffixRemoveChecked)
+				.OnCheckStateChanged(this, &SAdvancedRenamerPanel::OnSuffixRemoveCheckBoxChanged)
+			]
+			+ SCanvas::Slot()
+			.Position(FVector2D(LabelStart * 2.f, TitleOffsetY + LabelOffsetY + LineHeight))
+			.Size(LabelSize)
+			[
+				SNew(STextBlock)
+				.ColorAndOpacity(FStyleColors::AccentWhite.GetSpecifiedColor())
+				.Font(RegularFont)
+				.Text(LOCTEXT("SuffixRemove", "Remove Suffix"))
+			]
+			+ SCanvas::Slot()
+			.Position(FVector2D(EntryStart, TitleOffsetY + LineHeight))
+			.Size(SeparatorSize)
+			[
+				SAssignNew(SuffixSeparatorTextBox, SEditableTextBox)
+				.BackgroundColor(FStyleColors::Background.GetSpecifiedColor())
+				.ForegroundColor(FStyleColors::AccentWhite.GetSpecifiedColor())
+				.Font(RegularFont)
+				.IsEnabled(this, &SAdvancedRenamerPanel::IsSuffixRemoveSeparatorEnabled)
+				.Text(LOCTEXT("Underscore", "_"))
+				.OnVerifyTextChanged(this, &SAdvancedRenamerPanel::OnSuffixSeparatorVerifyTextChanged)
+				.OnTextChanged(this, &SAdvancedRenamerPanel::OnSuffixSeparatorChanged)
+			]
+			+ SCanvas::Slot()
+			.Position(FVector2D(EntryStart + SeparatorSize.X + 5.f, TitleOffsetY + LabelOffsetY + LineHeight))
+			.Size(LabelSize)
+			[
+				SNew(STextBlock)
+				.ColorAndOpacity(FStyleColors::AccentWhite.GetSpecifiedColor())
+				.Font(RegularFont)
+				.Text(LOCTEXT("Separator", "(separator)"))
 			]
 
-			+ SHorizontalBox::Slot()
-			.Padding(CancelButtonPadding)
-			.HAlign(HAlign_Right)
-			.AutoWidth()
+			// Remove first
+			+ SCanvas::Slot()
+			.Position(FVector2D(LabelStart, TitleOffsetY + LineHeight * 2.f) + RadioOffset)
+			.Size(CheckboxSize)
 			[
-				CreateCancelButton()
-			]);
-}
+				SAssignNew(SuffixRemoveCharactersCheckBox, SCheckBox)
+				.Style(&FAdvancedRenamerStyle::Get().GetWidgetStyle<FCheckBoxStyle>("AdvancedRenamer.Style.BlackRadioButton"))
+				.IsChecked(this, &SAdvancedRenamerPanel::IsSuffixRemoveCharactersChecked)
+				.OnCheckStateChanged(this, &SAdvancedRenamerPanel::OnSuffixRemoveCharactersCheckBoxChanged)
+			]
+			+ SCanvas::Slot()
+			.Position(FVector2D(LabelStart * 2.f, TitleOffsetY + LabelOffsetY + LineHeight * 2.f))
+			.Size(LabelSize)
+			[
+				SNew(STextBlock)
+				.ColorAndOpacity(FStyleColors::AccentWhite.GetSpecifiedColor())
+				.Font(RegularFont)
+				.Text(LOCTEXT("SuffixRemoveChars", "Remove Last"))
+			]
+			+ SCanvas::Slot()
+			.Position(FVector2D(EntryStart, TitleOffsetY + LineHeight * 2.f))
+			.Size(SpinSize1)
+			[
+				SAssignNew(SuffixRemoveCharactersSpinBox, SSpinBox<uint8>)
+				.Style(&FAppStyle::Get(), "Menu.SpinBox")
+				.Font(RegularFont)
+				.MinValue(1)
+				.MaxValue(9)
+				.Value(1)
+				.IsEnabled(this, &SAdvancedRenamerPanel::IsSuffixRemoveNumCharsEnabled)
+				.OnValueChanged(this, &SAdvancedRenamerPanel::OnSuffixRemoveCharactersChanged)
+			]
+			+ SCanvas::Slot()
+			.Position(FVector2D(EntryStart + SeparatorSize.X + 5.f, TitleOffsetY + LabelOffsetY + LineHeight * 2.f))
+			.Size(LabelSize)
+			[
+				SNew(STextBlock)
+				.ColorAndOpacity(FStyleColors::AccentWhite.GetSpecifiedColor())
+				.Font(RegularFont)
+				.Text(LOCTEXT("Characters", "character(s)"))
+			]
 
-TSharedRef<SWidget> SAdvancedRenamerPanel::CreateApplyButton()
-{
-	return SNew(SBox)
-		.HeightOverride(25)
-		.WidthOverride(75)
-		[
-			SNew(SButton)
-			.IsEnabled(this, &SAdvancedRenamerPanel::IsApplyButtonEnabled)
-			.ButtonStyle(&FAppStyle::Get().GetWidgetStyle<FButtonStyle>("PrimaryButton"))
-			.ContentPadding(FMargin(-2.f, -1.f))
-			.OnClicked(this, &SAdvancedRenamerPanel::OnApplyButtonClicked)
-			.Content()
-			[
-				SNew(SBox)
-				.VAlign(VAlign_Center)
-				[
-					SNew(STextBlock)
-					.Justification(ETextJustify::Center)
-					.Text(LOCTEXT("AR_Apply", "Apply"))
-				]
+			// Remove Number
+			+ SCanvas::Slot()
+			.Position(FVector2D(LabelStart, TitleOffsetY + LineHeight * 3.f) + CheckboxOffset)
+			.Size(CheckboxSize)
+			[
+				SAssignNew(SuffixRemoveNumberCheckBox, SCheckBox)
+				.IsChecked(this, &SAdvancedRenamerPanel::IsSuffixRemoveNumberChecked)
+				.OnCheckStateChanged(this, &SAdvancedRenamerPanel::OnSuffixRemoveNumberCheckBoxChanged)
+				.IsEnabled(this, &SAdvancedRenamerPanel::IsSuffixRemoveNumberCheckBoxEnabled)
+			]
+			+ SCanvas::Slot()
+			.Position(FVector2D(LabelStart * 2.f, TitleOffsetY + LabelOffsetY + LineHeight * 3.f))
+			.Size(LabelSize)
+			[
+				SNew(STextBlock)
+				.ColorAndOpacity(FStyleColors::AccentWhite.GetSpecifiedColor())
+				.Font(RegularFont)
+				.Text(LOCTEXT("SuffixRemoveNumber", "Remove Number"))
 			]
-		];
-}
 
-TSharedRef<SWidget> SAdvancedRenamerPanel::CreateResetButton()
-{
-	return SNew(SBox)
-		.HeightOverride(25)
-		.WidthOverride(75)
-		[
-			SNew(SButton)
-			.ContentPadding(FMargin(-2.f, -1.f))
-			.OnClicked(this, &SAdvancedRenamerPanel::OnResetButtonClicked)
-			.Content()
+			// Number
+			+ SCanvas::Slot()
+			.Position(FVector2D(LabelStart, TitleOffsetY + LineHeight * 4.f) + CheckboxOffset)
+			.Size(CheckboxSize)
 			[
-				SNew(SBox)
-				.VAlign(VAlign_Center)
-				[
-					SNew(STextBlock)
-					.Justification(ETextJustify::Center)
-					.Text(LOCTEXT("AR_Reset", "Reset"))
-				]
+				SAssignNew(SuffixNumberCheckBox, SCheckBox)
+				.IsChecked(this, &SAdvancedRenamerPanel::IsSuffixNumberChecked)
+				.OnCheckStateChanged(this, &SAdvancedRenamerPanel::OnSuffixNumberCheckBoxChanged)
+			]
+			+ SCanvas::Slot()
+			.Position(FVector2D(LabelStart * 2.f, TitleOffsetY + LabelOffsetY + LineHeight * 4.f))
+			.Size(LabelSize)
+			[
+				SNew(STextBlock)
+				.ColorAndOpacity(FStyleColors::AccentWhite.GetSpecifiedColor())
+				.Font(RegularFont)
+				.Text(LOCTEXT("SuffixNumber", "Add Number"))
+			]
+			+ SCanvas::Slot()
+			.Position(FVector2D(EntryStart , TitleOffsetY + LabelOffsetY + LineHeight * 4.f))
+			.Size(LabelSize)
+			[
+				SNew(STextBlock)
+				.ColorAndOpacity(FStyleColors::AccentWhite.GetSpecifiedColor())
+				.Font(RegularFont)
+				.Text(LOCTEXT("Start", "Start"))
+			]
+			+ SCanvas::Slot()
+			.Position(FVector2D(EntryStart + 33.f, TitleOffsetY + LineHeight * 4.f))
+			.Size(SpinSize3)
+			[
+
+				SAssignNew(SuffixNumberStartSpinBox, SSpinBox<int32>)
+				.Style(&FAppStyle::Get(), "Menu.SpinBox")
+				.Font(RegularFont)
+				.MinValue(0)
+				.MaxValue(999)
+				.Value(1)
+				.IsEnabled(false)
+				.OnValueChanged(this, &SAdvancedRenamerPanel::OnSuffixNumberStartChanged)
+			]
+			+ SCanvas::Slot()
+			.Position(FVector2D(EntryStart + 33.f + SpinSize3.X + 5.f, TitleOffsetY + LabelOffsetY + LineHeight * 4.f))
+			.Size(LabelSize)
+			[
+				SNew(STextBlock)
+				.ColorAndOpacity(FStyleColors::AccentWhite.GetSpecifiedColor())
+				.Font(RegularFont)
+				.Text(LOCTEXT("Step", "Step"))
+			]
+			+ SCanvas::Slot()
+			.Position(FVector2D(EntryStart + 33.f + SpinSize3.X + 36.f, TitleOffsetY + LineHeight * 4.f))
+			.Size(SpinSize2)
+			[
+
+				SAssignNew(SuffixNumberStepSpinBox, SSpinBox<int32>)
+				.Style(&FAppStyle::Get(), "Menu.SpinBox")
+				.Font(RegularFont)
+				.MinValue(1)
+				.MaxValue(99)
+				.Value(1)
+				.IsEnabled(false)
+				.OnValueChanged(this, &SAdvancedRenamerPanel::OnSuffixNumberStepChanged)
 			]
 		];
 }
 
-TSharedRef<SWidget> SAdvancedRenamerPanel::CreateCancelButton()
+TSharedRef<SWidget> SAdvancedRenamerPanel::CreateSearchAndReplace()
 {
-	return SNew(SBox)
-		.HeightOverride(25)
-		.WidthOverride(75)
+	using namespace UE::AdvancedRenamer::Private;
+
+	return SNew(SBorder)
+		.BorderImage(FAppStyle::Get().GetBrush("DetailsView.CategoryTop"))
+		.Content()
 		[
-			SNew(SButton)
-			.ContentPadding(FMargin(-2.f, -1.f))
-			.OnClicked(this, &SAdvancedRenamerPanel::OnCancelButtonClicked)
-			.Content()
+			SNew(SCanvas)
+
+			// Title
+			+ SCanvas::Slot()
+			.Position(FVector2D(0.f, LabelOffsetY))
+			.Size(TitleSize)
 			[
-				SNew(SBox)
-				.VAlign(VAlign_Center)
-				[
-					SNew(STextBlock)
-					.Justification(ETextJustify::Center)
-					.Text(LOCTEXT("AR_Cancel", "Cancel"))
-				]
+				SNew(STextBlock)
+				.ColorAndOpacity(FStyleColors::AccentBlue.GetSpecifiedColor())
+				.Font(TitleFont)
+				.Text(LOCTEXT("SearchReplaceTitle", "Search and Replace"))
+			]
+
+			+ SCanvas::Slot()
+			.Position(FVector2D(0.f, TitleOffsetY) + RadioOffset)
+			.Size(CheckboxSize)
+			[
+				SAssignNew(SearchReplacePlainTextCheckbox, SCheckBox)
+				.Style(&FAdvancedRenamerStyle::Get().GetWidgetStyle<FCheckBoxStyle>("AdvancedRenamer.Style.BlackRadioButton"))
+				.IsChecked(this, &SAdvancedRenamerPanel::IsSearchReplacePlainTextChecked)
+				.OnCheckStateChanged(this, &SAdvancedRenamerPanel::OnSearchReplacePlainTextCheckBoxChanged)
+			]
+
+			+ SCanvas::Slot()
+			.Position(FVector2D(0.f + LabelStart, TitleOffsetY + LabelOffsetY))
+			.Size(LabelSize)
+			[
+				SNew(STextBlock)
+				.ColorAndOpacity(FStyleColors::AccentWhite.GetSpecifiedColor())
+				.Font(RegularFont)
+				.Text(LOCTEXT("PlainText", "Plain Text"))
+			]
+
+			+ SCanvas::Slot()
+			.Position(FVector2D(92.f, TitleOffsetY) + RadioOffset)
+			.Size(CheckboxSize)
+			[
+				SAssignNew(SearchReplaceRegexCheckbox, SCheckBox)
+				.Style(&FAdvancedRenamerStyle::Get().GetWidgetStyle<FCheckBoxStyle>("AdvancedRenamer.Style.BlackRadioButton"))
+				.IsChecked(this, &SAdvancedRenamerPanel::IsSearchReplaceRegexChecked)
+				.OnCheckStateChanged(this, &SAdvancedRenamerPanel::OnSearchReplaceRegexCheckBoxChanged)
+			]
+
+			+ SCanvas::Slot()
+			.Position(FVector2D(92.f + LabelStart, TitleOffsetY + LabelOffsetY))
+			.Size(LabelSize)
+			[
+				SNew(STextBlock)
+				.ColorAndOpacity(FStyleColors::AccentWhite.GetSpecifiedColor())
+				.Font(RegularFont)
+				.Text(LOCTEXT("Regex", "Regex"))
+			]
+
+			+ SCanvas::Slot()
+			.Position(FVector2D(200.f, TitleOffsetY) + CheckboxOffset)
+			.Size(CheckboxSize)
+			[
+				SAssignNew(SearchReplaceIgnoreCaseCheckBox, SCheckBox)
+				.IsChecked(this, &SAdvancedRenamerPanel::IsSearchReplaceIgnoreCaseChecked)
+				.OnCheckStateChanged(this, &SAdvancedRenamerPanel::OnSearchReplaceIgnoreCaseCheckBoxChanged)
+			]
+
+			+ SCanvas::Slot()
+			.Position(FVector2D(200.f + LabelStart, TitleOffsetY + LabelOffsetY))
+			.Size(LabelSize)
+			[
+				SNew(STextBlock)
+				.ColorAndOpacity(FStyleColors::AccentWhite.GetSpecifiedColor())
+				.Font(RegularFont)
+				.Text(LOCTEXT("IgnoreCase", "Ignore Case"))
+			]
+
+			+ SCanvas::Slot()
+			.Position(FVector2D(0.f, TitleOffsetY + LineHeight))
+			.Size(FVector2D(LeftBlockWidth - 4.f, 54.f))
+			[
+				SAssignNew(SearchReplaceSearchTextBox, SMultiLineEditableTextBox)
+				.BackgroundColor(FStyleColors::Background.GetSpecifiedColor())
+				.ForegroundColor(FStyleColors::AccentWhite.GetSpecifiedColor())
+				.Font(RegularFont)
+				.AllowMultiLine(false)
+				.AutoWrapText(true)
+				.HintText(LOCTEXT("RegeSearchHint", "Search"))
+				.WrappingPolicy(ETextWrappingPolicy::AllowPerCharacterWrapping)
+				.OnTextChanged(this, &SAdvancedRenamerPanel::OnSearchReplaceSearchTextChanged)
+			]
+
+			+ SCanvas::Slot()
+			.Position(FVector2D(0.f, TitleOffsetY + LineHeight + 59.f))
+			.Size(FVector2D(LeftBlockWidth - 4.f, 54.f))
+			[
+				SAssignNew(SearchReplaceReplaceTextBox, SMultiLineEditableTextBox)
+				.BackgroundColor(FStyleColors::Background.GetSpecifiedColor())
+				.ForegroundColor(FStyleColors::AccentWhite.GetSpecifiedColor())
+				.Font(RegularFont)
+				.AllowMultiLine(false)
+				.AutoWrapText(true)
+				.HintText(LOCTEXT("RegeReplaceHint", "Replace"))
+				.WrappingPolicy(ETextWrappingPolicy::AllowPerCharacterWrapping)
+				.OnTextChanged(this, &SAdvancedRenamerPanel::OnSearchReplaceReplaceTextChanged)
 			]
 		];
 }
 
-void SAdvancedRenamerPanel::CreateRightPanel()
+void SAdvancedRenamerPanel::CreateRightPane(const TSharedRef<SCanvas>& InCanvas)
 {
-	using namespace AdvancedRenamerSlateUtils::Default;
+	using namespace UE::AdvancedRenamer::Private;
 
-	RightSideBox->SetContent(SNew(SVerticalBox)
-		+ SVerticalBox::Slot()
+	InCanvas->AddSlot()
+		.Position(FVector2D(RightBlockOffsetX, 5.f))
+		.Size(this, &SAdvancedRenamerPanel::GetRightPaneSize)
 		[
 			CreateRenamePreview()
-		]);
+		];
 }
 
 TSharedRef<SWidget> SAdvancedRenamerPanel::CreateRenamePreview()
 {
-	using namespace AdvancedRenamerSlateUtils::Default;
+	using namespace UE::AdvancedRenamer::Private;
 
 	return SNew(SBorder)
-		.BorderImage(FAdvancedRenamerStyle::Get().GetBrush("AdvancedRenamer.Style.BackgroundBorder"))
+		.BorderImage(FAppStyle::Get().GetBrush("DetailsView.CategoryTop"))
+		.HAlign(EHorizontalAlignment::HAlign_Fill)
+		.VAlign(EVerticalAlignment::VAlign_Fill)
 		.Content()
 		[
 			SNew(SVerticalBox)
 
+			+ SVerticalBox::Slot()
+			.Padding(0.f, LabelOffsetY, 0.f, 0.f)
+			.AutoHeight()
+			[
+				SNew(STextBlock)
+				.ColorAndOpacity(FStyleColors::AccentBlue.GetSpecifiedColor())
+				.Font(TitleFont)
+				.Text(LOCTEXT("ObjectRenamePreviewTitle", "Rename Preview"))
+			]
+
 			+ SVerticalBox::Slot()
 			.FillHeight(1.f)
+			.HAlign(EHorizontalAlignment::HAlign_Fill)
 			[
 				SAssignNew(RenamePreviewListBox, SBox)
 				.Content()
 				[
 					SAssignNew(RenamePreviewList, SListView<TSharedPtr<FAdvancedRenamerPreview>>)
-					.ListViewStyle(&FAdvancedRenamerStyle::Get().GetWidgetStyle<FTableViewStyle>("AdvancedRenamer.Style.ListView"))
 					.ItemHeight(ListLineHeight)
 					.ListItemsSource(&Renamer->GetPreviews())
 					.OnGenerateRow(this, &SAdvancedRenamerPanel::OnGenerateRowForList)
 					.HeaderRow(
-
 						SAssignNew(RenamePreviewListHeaderRow, SHeaderRow)
-						.Style(&FAdvancedRenamerStyle::Get().GetWidgetStyle<FHeaderRowStyle>("AdvancedRenamer.Style.HeaderRow"))
-
 						+ SHeaderRow::Column(OriginalNameColumnName)
-						.HeaderContentPadding(FMargin(8.f, 2.f))
-						.DefaultLabel(LOCTEXT("AR_Old", "Old"))
+						.DefaultLabel(LOCTEXT("From", "From"))
 						.FillWidth(0.5f)
-
 						+ SHeaderRow::Column(NewNameColumnName)
-						.HeaderContentPadding(FMargin(8.f, 2.f))
-						.DefaultLabel(LOCTEXT("AR_New", "New"))
+						.DefaultLabel(LOCTEXT("To", "To"))
 						.FillWidth(0.5f)
 					)
 					.OnKeyDownHandler(this, &SAdvancedRenamerPanel::OnListViewKeyDown)
 					.OnContextMenuOpening(this, &SAdvancedRenamerPanel::GenerateListViewContextMenu)
 				]
 			]
+
+			+ SVerticalBox::Slot()
+			.AutoHeight()
+			.HAlign(EHorizontalAlignment::HAlign_Fill)
+			[
+				SAssignNew(ApplyButton, SButton)
+				.ButtonStyle(FAdvancedRenamerStyle::Get(), "AdvancedRenamer.Style.DarkButton")
+				.ToolTipText(LOCTEXT("ApplyRenameToolTip", "Renames all actors."))
+				.HAlign(EHorizontalAlignment::HAlign_Center)
+				.VAlign(EVerticalAlignment::VAlign_Center)
+				.IsEnabled(this, &SAdvancedRenamerPanel::IsApplyButtonEnabled)
+				.OnClicked(this, &SAdvancedRenamerPanel::OnApplyButtonClicked)
+				[
+					SNew(STextBlock)
+					.ColorAndOpacity(FStyleColors::AccentWhite.GetSpecifiedColor())
+					.Font(RegularFont)
+					.Text(LOCTEXT("ApplyRename", "Apply"))
+				]
+			]
 		];
 }
 
@@ -362,7 +782,7 @@ bool SAdvancedRenamerPanel::CloseWindow()
 	return false;
 }
 
-void SAdvancedRenamerPanel::RefreshListViewAndUpdate(const double InCurrentTime)
+void SAdvancedRenamerPanel::RefreshListView(const double InCurrentTime)
 {
 	const int32 CurrentCount = Renamer->Num();
 
@@ -373,11 +793,40 @@ void SAdvancedRenamerPanel::RefreshListViewAndUpdate(const double InCurrentTime)
 		RenamePreviewList->RequestListRefresh();
 	}
 
+	MinDesiredOriginalNameWidth = 0.f;
+	MinDesiredNewNameWidth = 0.f;
 	RenamePreviewList->RebuildList();
 
 	ListLastUpdateTime = InCurrentTime;
 }
 
+void SAdvancedRenamerPanel::UpdateRequiredListWidth()
+{
+	if (!RenamePreviewListBox.IsValid() || !RenamePreviewListHeaderRow.IsValid())
+	{
+		return;
+	}
+
+	if (MinDesiredNewNameWidth == 0.f)
+	{
+		return;
+	}
+
+	using namespace UE::AdvancedRenamer::Private;
+
+	const float ActualWidth = FMath::Max(GetTickSpaceGeometry().GetLocalSize().X, MinDesiredOriginalNameWidth + MinDesiredNewNameWidth + 20.f);
+
+	RenamePreviewListHeaderRow->SetColumnWidth(
+		OriginalNameColumnName, 
+		(MinDesiredOriginalNameWidth + 10.f) / (MinDesiredOriginalNameWidth + MinDesiredNewNameWidth + 20.f) * ActualWidth
+	);
+
+	RenamePreviewListHeaderRow->SetColumnWidth(
+		NewNameColumnName, 
+		(MinDesiredNewNameWidth + 10.f) / (MinDesiredOriginalNameWidth + MinDesiredNewNameWidth + 20.f) * ActualWidth
+	);
+}
+
 void SAdvancedRenamerPanel::RemoveSelectedObjects()
 {
 	if (!RenamePreviewList.IsValid() || RenamePreviewList->GetNumItemsSelected() == 0)
@@ -440,10 +889,337 @@ void SAdvancedRenamerPanel::Tick(const FGeometry& InAllottedGeometry, const doub
 	}
 	else if (Renamer->IsDirty() && InCurrentTime >= (ListLastUpdateTime + MinUpdateFrequency))
 	{
-		RefreshListViewAndUpdate(InCurrentTime);
+		RefreshListView(InCurrentTime);
+	}
+	else
+	{
+		UpdateRequiredListWidth();
 	}
 }
 
+void SAdvancedRenamerPanel::OnBaseNameChanged(const FText& InNewText)
+{
+	Renamer->GetOptions().BaseName = InNewText.ToString();
+	Renamer->MarkDirty();
+}
+
+void SAdvancedRenamerPanel::OnPrefixChanged(const FText& InNewText)
+{
+	Renamer->GetOptions().AddPrefix = InNewText.ToString();
+	Renamer->MarkDirty();
+}
+
+ECheckBoxState SAdvancedRenamerPanel::IsPrefixRemoveChecked() const
+{
+	return bRemovePrefixSeparator ? ECheckBoxState::Checked : ECheckBoxState::Unchecked;
+}
+
+void SAdvancedRenamerPanel::OnPrefixRemoveCheckBoxChanged(ECheckBoxState InNewState)
+{
+	bRemovePrefixSeparator = InNewState == ECheckBoxState::Checked;
+
+	if (bRemovePrefixSeparator)
+	{
+		bRemovePrefixNumChars = false;
+	}
+
+	if (bRemovePrefixSeparator && PrefixSeparatorTextBox.IsValid())
+	{
+		Renamer->GetOptions().RemovePrefixSeparator = PrefixSeparatorTextBox->GetText().ToString();
+		Renamer->GetOptions().RemovePrefixCharacterCount = 0;
+	}
+	else
+	{
+		Renamer->GetOptions().RemovePrefixSeparator = "";
+	}
+
+	Renamer->MarkDirty();
+}
+
+bool SAdvancedRenamerPanel::IsPrefixRemoveSeparatorEnabled() const
+{
+	return bRemovePrefixSeparator;
+}
+
+bool SAdvancedRenamerPanel::OnPrefixSeparatorVerifyTextChanged(const FText& InText, FText& OutErrorText) const
+{
+	if (InText.ToString().Len() > 1)
+	{
+		OutErrorText = LOCTEXT("SeparatorError", "Separators can only be a single character.");
+		return false;
+	}
+
+	return true;
+}
+
+void SAdvancedRenamerPanel::OnPrefixSeparatorChanged(const FText& InNewText)
+{
+	Renamer->GetOptions().RemovePrefixSeparator = InNewText.ToString();
+	Renamer->MarkDirty();
+}
+
+ECheckBoxState SAdvancedRenamerPanel::IsPrefixRemoveCharactersChecked() const
+{
+	return bRemovePrefixNumChars ? ECheckBoxState::Checked : ECheckBoxState::Unchecked;
+}
+
+void SAdvancedRenamerPanel::OnPrefixRemoveCharactersCheckBoxChanged(ECheckBoxState InNewState)
+{
+	bRemovePrefixNumChars = InNewState == ECheckBoxState::Checked;
+
+	if (bRemovePrefixNumChars)
+	{
+		bRemovePrefixSeparator = false;
+	}
+
+	if (bRemovePrefixNumChars && PrefixRemoveCharactersSpinBox.IsValid())
+	{
+		Renamer->GetOptions().RemovePrefixCharacterCount = PrefixRemoveCharactersSpinBox->GetValue();
+		Renamer->GetOptions().RemovePrefixSeparator = "";
+	}
+	else
+	{
+		Renamer->GetOptions().RemovePrefixCharacterCount = 0;
+	}
+
+	Renamer->MarkDirty();
+}
+
+bool SAdvancedRenamerPanel::IsPrefixRemoveNumCharsEnabled() const
+{
+	return bRemovePrefixNumChars;
+}
+
+void SAdvancedRenamerPanel::OnPrefixRemoveCharactersChanged(uint8 InNewValue)
+{
+	Renamer->GetOptions().RemovePrefixCharacterCount = InNewValue;
+	Renamer->MarkDirty();
+}
+
+void SAdvancedRenamerPanel::OnSuffixChanged(const FText& InNewText)
+{
+	Renamer->GetOptions().AddSuffix = InNewText.ToString();
+	Renamer->MarkDirty();
+}
+
+ECheckBoxState SAdvancedRenamerPanel::IsSuffixRemoveChecked() const
+{
+	return bRemoveSuffixSeparator ? ECheckBoxState::Checked : ECheckBoxState::Unchecked;
+}
+
+void SAdvancedRenamerPanel::OnSuffixRemoveCheckBoxChanged(ECheckBoxState InNewState)
+{
+	bRemoveSuffixSeparator = InNewState == ECheckBoxState::Checked;
+
+	if (bRemoveSuffixSeparator)
+	{
+		bRemoveSuffixNumChars = false;
+	}
+
+	if (bRemoveSuffixSeparator && SuffixSeparatorTextBox.IsValid())
+	{
+		Renamer->GetOptions().RemoveSuffixSeparator = SuffixSeparatorTextBox->GetText().ToString();
+		Renamer->GetOptions().RemoveSuffixCharacterCount = 0;
+	}
+	else
+	{
+		Renamer->GetOptions().RemoveSuffixSeparator = "";
+	}
+
+	Renamer->MarkDirty();
+}
+
+bool SAdvancedRenamerPanel::IsSuffixRemoveSeparatorEnabled() const
+{
+	return bRemoveSuffixSeparator;
+}
+
+bool SAdvancedRenamerPanel::OnSuffixSeparatorVerifyTextChanged(const FText& InText, FText& OutErrorText) const
+{
+	if (InText.ToString().Len() > 1)
+	{
+		OutErrorText = LOCTEXT("SeparatorError", "Separators can only be a single character.");
+		return false;
+	}
+
+	return true;
+}
+
+void SAdvancedRenamerPanel::OnSuffixSeparatorChanged(const FText& InNewText)
+{
+	Renamer->GetOptions().RemoveSuffixSeparator = InNewText.ToString();
+	Renamer->MarkDirty();
+}
+
+ECheckBoxState SAdvancedRenamerPanel::IsSuffixRemoveCharactersChecked() const
+{
+	return bRemoveSuffixNumChars ? ECheckBoxState::Checked : ECheckBoxState::Unchecked;
+}
+
+bool SAdvancedRenamerPanel::IsSuffixRemoveNumCharsEnabled() const
+{
+	return bRemoveSuffixNumChars;
+}
+
+void SAdvancedRenamerPanel::OnSuffixRemoveCharactersCheckBoxChanged(ECheckBoxState InNewState)
+{
+	bRemoveSuffixNumChars = InNewState == ECheckBoxState::Checked;
+
+	if (bRemoveSuffixNumChars)
+	{
+		bRemoveSuffixSeparator = false;
+	}
+
+	if (bRemoveSuffixNumChars && SuffixRemoveCharactersSpinBox.IsValid())
+	{
+		Renamer->GetOptions().RemoveSuffixCharacterCount = SuffixRemoveCharactersSpinBox->GetValue();
+		Renamer->GetOptions().RemoveSuffixSeparator = "";
+	}
+	else
+	{
+		Renamer->GetOptions().RemoveSuffixCharacterCount = 0;
+	}
+
+	Renamer->MarkDirty();
+}
+
+void SAdvancedRenamerPanel::OnSuffixRemoveCharactersChanged(uint8 InNewValue)
+{
+	Renamer->GetOptions().RemoveSuffixCharacterCount = InNewValue;
+	Renamer->MarkDirty();
+}
+
+ECheckBoxState SAdvancedRenamerPanel::IsSuffixRemoveNumberChecked() const
+{
+	return Renamer->GetOptions().bRemoveSuffixNumber ? ECheckBoxState::Checked : ECheckBoxState::Unchecked;
+}
+
+void SAdvancedRenamerPanel::OnSuffixRemoveNumberCheckBoxChanged(ECheckBoxState InNewState)
+{
+	Renamer->GetOptions().bRemoveSuffixNumber = InNewState == ECheckBoxState::Checked;
+
+	if (!Renamer->GetOptions().bRemoveSuffixNumber)
+	{
+		Renamer->GetOptions().bAddSuffixNumber = false;
+	}
+
+	Renamer->MarkDirty();
+}
+
+ECheckBoxState SAdvancedRenamerPanel::IsSuffixNumberChecked() const
+{
+	return Renamer->GetOptions().bAddSuffixNumber ? ECheckBoxState::Checked : ECheckBoxState::Unchecked;
+}
+
+bool SAdvancedRenamerPanel::IsSuffixRemoveNumberCheckBoxEnabled() const
+{
+	return !Renamer->GetOptions().bAddSuffixNumber;
+}
+
+void SAdvancedRenamerPanel::OnSuffixNumberCheckBoxChanged(ECheckBoxState InNewState)
+{
+	Renamer->GetOptions().bAddSuffixNumber = InNewState == ECheckBoxState::Checked;
+
+	if (Renamer->GetOptions().bAddSuffixNumber)
+	{
+		Renamer->GetOptions().bRemoveSuffixNumber = true;
+	}
+
+	Renamer->MarkDirty();
+}
+
+void SAdvancedRenamerPanel::OnSuffixNumberStartChanged(int32 InNewValue)
+{
+	Renamer->GetOptions().AddSuffixNumberStart = InNewValue;
+	Renamer->MarkDirty();
+}
+
+void SAdvancedRenamerPanel::OnSuffixNumberStepChanged(int32 InNewValue)
+{
+	Renamer->GetOptions().AddSuffixNumberStep = InNewValue;
+	Renamer->MarkDirty();
+}
+
+ECheckBoxState SAdvancedRenamerPanel::IsSearchReplacePlainTextChecked() const
+{
+	return Renamer->GetOptions().SearchAndReplaceType == EAdvancedRenamerSeachAndReplaceType::PlainText ? ECheckBoxState::Checked : ECheckBoxState::Unchecked;
+}
+
+void SAdvancedRenamerPanel::OnSearchReplacePlainTextCheckBoxChanged(ECheckBoxState InNewState)
+{
+	if (InNewState != ECheckBoxState::Checked)
+	{
+		return;
+	}
+
+	Renamer->GetOptions().SearchAndReplaceType = EAdvancedRenamerSeachAndReplaceType::PlainText;
+	Renamer->MarkDirty();
+}
+
+ECheckBoxState SAdvancedRenamerPanel::IsSearchReplaceRegexChecked() const
+{
+	return Renamer->GetOptions().SearchAndReplaceType == EAdvancedRenamerSeachAndReplaceType::RegularExpression ? ECheckBoxState::Checked : ECheckBoxState::Unchecked;
+}
+
+void SAdvancedRenamerPanel::OnSearchReplaceRegexCheckBoxChanged(ECheckBoxState InNewState)
+{
+	if (InNewState != ECheckBoxState::Checked)
+	{
+		return;
+	}
+
+	Renamer->GetOptions().SearchAndReplaceType = EAdvancedRenamerSeachAndReplaceType::RegularExpression;
+	Renamer->MarkDirty();
+}
+
+ECheckBoxState SAdvancedRenamerPanel::IsSearchReplaceIgnoreCaseChecked() const
+{
+	return Renamer->GetOptions().SearchAndReplaceCase == ESearchCase::IgnoreCase ? ECheckBoxState::Checked : ECheckBoxState::Unchecked;
+}
+
+void SAdvancedRenamerPanel::OnSearchReplaceIgnoreCaseCheckBoxChanged(ECheckBoxState InNewState)
+{
+	switch (InNewState)
+	{
+		case ECheckBoxState::Checked:
+			Renamer->GetOptions().SearchAndReplaceCase = ESearchCase::IgnoreCase;
+			break;
+
+		case ECheckBoxState::Unchecked:
+			Renamer->GetOptions().SearchAndReplaceCase = ESearchCase::CaseSensitive;
+			break;
+
+		default:
+			return;
+	}
+
+	Renamer->MarkDirty();
+}
+
+void SAdvancedRenamerPanel::OnSearchReplaceSearchTextChanged(const FText& InNewText)
+{
+	Renamer->GetOptions().SearchAndReplaceFromText = InNewText.ToString();
+
+	if (Renamer->GetOptions().SearchAndReplaceType == EAdvancedRenamerSeachAndReplaceType::None)
+	{
+		Renamer->GetOptions().SearchAndReplaceType = EAdvancedRenamerSeachAndReplaceType::PlainText;
+	}
+
+	Renamer->MarkDirty();
+}
+
+void SAdvancedRenamerPanel::OnSearchReplaceReplaceTextChanged(const FText& InNewText)
+{
+	Renamer->GetOptions().SearchAndReplaceToText = InNewText.ToString();
+
+	if (Renamer->GetOptions().SearchAndReplaceType == EAdvancedRenamerSeachAndReplaceType::None)
+	{
+		Renamer->GetOptions().SearchAndReplaceType = EAdvancedRenamerSeachAndReplaceType::PlainText;
+	}
+
+	Renamer->MarkDirty();
+}
+
 TSharedRef<ITableRow> SAdvancedRenamerPanel::OnGenerateRowForList(TSharedPtr<FAdvancedRenamerPreview> InItem, 
 	const TSharedRef<STableViewBase>& InOwnerTable)
 {
@@ -490,43 +1266,50 @@ bool SAdvancedRenamerPanel::IsApplyButtonEnabled() const
 
 FReply SAdvancedRenamerPanel::OnApplyButtonClicked()
 {
-	FScopedTransaction Transaction(LOCTEXT("AdvancedRenamerRename", "Batch Renamer Rename"));
+	FScopedTransaction Transaction(LOCTEXT("AdvancedRenamerRename", "Advanced Renamer Rename"));
 
-	if (!Renamer->Execute())
+	if (Renamer->Execute())
 	{
-		FNotificationInfo Info(LOCTEXT("AR_RenameIssueNotification", "Errors occured while applying the rename.\nSee Output Log for more information on why the rename failed."));
-		Info.ExpireDuration = 5.0f;
-		const TSharedPtr<SNotificationItem> Notification = FSlateNotificationManager::Get().AddNotification(Info);
-		if (Notification.IsValid())
-		{
-			Notification->SetCompletionState(SNotificationItem::CS_Fail);
-		}
+		CloseWindow();
 	}
 
-	CloseWindow();
 	return FReply::Handled();
 }
 
-FReply SAdvancedRenamerPanel::OnResetButtonClicked()
+FVector2D SAdvancedRenamerPanel::GetRightPaneSize() const
 {
-	TArray<TSharedPtr<IAdvancedRenamerSection>> RegisteredSections = IAdvancedRenamerModule::Get().GetRegisteredSections();
-	for (const TSharedPtr<IAdvancedRenamerSection>& Section : RegisteredSections)
-	{
-		Section->ResetToDefault();
-	}
+	using namespace UE::AdvancedRenamer::Private;
 
-	if (Renamer.IsValid())
-	{
-		Renamer->MarkDirty();
-	}
+	const FVector2f GeoSize = GetTickSpaceGeometry().GetLocalSize();
 
-	return FReply::Handled();
+	return FVector2D(
+		FMath::Max(RightBlockSize.X, GeoSize.X - 319.0),
+		RightBlockSize.Y
+	);
 }
 
-FReply SAdvancedRenamerPanel::OnCancelButtonClicked()
+FVector2D SAdvancedRenamerPanel::GetListViewsize() const
 {
-	CloseWindow();
-	return FReply::Handled();
+	using namespace UE::AdvancedRenamer::Private;
+
+	const FVector2f GeoSize = GetTickSpaceGeometry().GetLocalSize();
+
+	return FVector2D(
+		FMath::Max(ListViewSize.X, GeoSize.X - 323.0),
+		ListViewSize.Y
+	);
+}
+
+FVector2D SAdvancedRenamerPanel::GetApplyButtonSize() const
+{
+	using namespace UE::AdvancedRenamer::Private;
+
+	const FVector2f GeoSize = GetTickSpaceGeometry().GetLocalSize();
+
+	return FVector2D(
+		FMath::Max(ListViewSize.X, GeoSize.X - 323.0),
+		ApplyButtonHeight
+	);
 }
 
 #undef LOCTEXT_NAMESPACE
diff --git a/Engine/Plugins/Experimental/AdvancedRenamer/Source/AdvancedRenamer/Private/Slate/SAdvancedRenamerPanel.h b/Engine/Plugins/Experimental/AdvancedRenamer/Source/AdvancedRenamer/Private/Slate/SAdvancedRenamerPanel.h
--- a/Engine/Plugins/Experimental/AdvancedRenamer/Source/AdvancedRenamer/Private/Slate/SAdvancedRenamerPanel.h
+++ b/Engine/Plugins/Experimental/AdvancedRenamer/Source/AdvancedRenamer/Private/Slate/SAdvancedRenamerPanel.h
@@ -2,23 +2,24 @@
 
 #pragma once
 
-#include "Input/Reply.h"
-#include "Widgets/SCompoundWidget.h"
-#include "Templates/SharedPointerFwd.h"
+#include "Styling/SlateTypes.h"
+#include "UObject/StrongObjectPtr.h"
+#include "Widgets/Views/SListView.h"
+#include "Widgets/Views/STableRow.h"
+#include "Widgets/Views/STableViewBase.h"
 
 class FUICommandList;
 class IAdvancedRenamer;
-class ITableRow;
 class SBox;
+class SButton;
+class SCanvas;
+class SCheckBox;
+class SEditableTextBox;
 class SHeaderRow;
-template <typename InItemType>
-class SListView;
-class STableViewBase;
-class SVerticalBox;
-class SWidget;
+class SMultiLineEditableTextBox; 
+class UObject;
 struct FAdvancedRenamerPreview;
-struct FGeometry;
-struct FKeyEvent;
+template<typename NumericType> class SSpinBox;
 
 class SAdvancedRenamerPanel : public SCompoundWidget
 {
@@ -30,87 +31,132 @@ public:
 
 	void Construct(const FArguments& InArgs, const TSharedRef<IAdvancedRenamer>& InRenamer);
 
-private:
-	/** Create Apply Reset and Cancel bottom panel */
-	void CreateApplyResetCancelPanel();
+protected:
+	static constexpr double MinUpdateFrequency = 0.1;
 
-	/** Create Apply button */
-	TSharedRef<SWidget> CreateApplyButton();
-
-	/** Create Reset button */
-	TSharedRef<SWidget> CreateResetButton();
+	TSharedPtr<IAdvancedRenamer> Renamer;
+	TSharedPtr<FUICommandList> CommandList;
 
-	/** Create Cancel button */
-	TSharedRef<SWidget> CreateCancelButton();
+	double ListLastUpdateTime = 0;
+	float MinDesiredOriginalNameWidth = 0.f;
+	float MinDesiredNewNameWidth = 0.f;
+
+	bool bRemovePrefixSeparator;
+	bool bRemovePrefixNumChars;
+	bool bRemoveSuffixSeparator;
+	bool bRemoveSuffixNumChars;
+
+	TSharedPtr<SEditableTextBox> BaseNameTextBox;
+	TSharedPtr<SEditableTextBox> PrefixTextBox;
+	TSharedPtr<SCheckBox> PrefixRemoveCheckBox;
+	TSharedPtr<SEditableTextBox> PrefixSeparatorTextBox;
+	TSharedPtr<SCheckBox> PrefixRemoveCharactersCheckBox;
+	TSharedPtr<SSpinBox<uint8>> PrefixRemoveCharactersSpinBox;
+	TSharedPtr<SEditableTextBox> SuffixTextBox;
+	TSharedPtr<SCheckBox> SuffixRemoveCheckBox;
+	TSharedPtr<SEditableTextBox> SuffixSeparatorTextBox;
+	TSharedPtr<SCheckBox> SuffixRemoveCharactersCheckBox;
+	TSharedPtr<SSpinBox<uint8>> SuffixRemoveCharactersSpinBox;
+	TSharedPtr<SCheckBox> SuffixRemoveNumberCheckBox;
+	TSharedPtr<SCheckBox> SuffixNumberCheckBox;
+	TSharedPtr<SSpinBox<int32>> SuffixNumberStartSpinBox;
+	TSharedPtr<SSpinBox<int32>> SuffixNumberStepSpinBox;
+	TSharedPtr<SCheckBox> SearchReplacePlainTextCheckbox;
+	TSharedPtr<SCheckBox> SearchReplaceRegexCheckbox;
+	TSharedPtr<SCheckBox> SearchReplaceIgnoreCaseCheckBox;
+	TSharedPtr<SMultiLineEditableTextBox> SearchReplaceSearchTextBox;
+	TSharedPtr<SMultiLineEditableTextBox> SearchReplaceReplaceTextBox;
+	TSharedPtr<SBox> RenamePreviewListBox;
+	TSharedPtr<SHeaderRow> RenamePreviewListHeaderRow;
+	TSharedPtr<SListView<TSharedPtr<FAdvancedRenamerPreview>>> RenamePreviewList;
+	TSharedPtr<SButton> ApplyButton;
 
-	/** Create the right side panel */
-	void CreateRightPanel();
+	void CreateLeftPane(const TSharedRef<SCanvas>& InCanvas);
+	TSharedRef<SWidget> CreateBaseName();
+	TSharedRef<SWidget> CreatePrefix();
+	TSharedRef<SWidget> CreateSuffix();
+	TSharedRef<SWidget> CreateSearchAndReplace();
 
-	/** Create the rename preview */
+	void CreateRightPane(const TSharedRef<SCanvas>& InCanvas);
 	TSharedRef<SWidget> CreateRenamePreview();
 
-	/** Close the Renamer window */
 	bool CloseWindow();
 
-	/** Refresh and Update the Preview */
-	void RefreshListViewAndUpdate(const double InCurrentTime);
+	void RefreshListView(const double InCurrentTime);
+	void UpdateRequiredListWidth();
 
-	/** Remove selected entries from the Renamer*/
 	void RemoveSelectedObjects();
 
-	//~ Begin SWidget Interface
 	virtual void Tick(const FGeometry& InAllottedGeometry, const double InCurrentTime, const float InDeltaTime) override;
-	//~ End SWidget Interface
 
-	/** Callback to generate ListBoxRows */
-	TSharedRef<ITableRow> OnGenerateRowForList(TSharedPtr<FAdvancedRenamerPreview> InItem, const TSharedRef<STableViewBase>& InOwnerTable);
+	void OnBaseNameChanged(const FText& InNewText);
 
-	/** Handle the click on the ListView */
-	FReply OnListViewKeyDown(const FGeometry& InGeometry, const FKeyEvent& InKeyEvent);
+	void OnPrefixChanged(const FText& InNewText);
 
-	/** Generate the ListView menu */
-	TSharedPtr<SWidget> GenerateListViewContextMenu();
+	ECheckBoxState IsPrefixRemoveChecked() const;
+	void OnPrefixRemoveCheckBoxChanged(ECheckBoxState InNewState);
 
-	/** Whether or not the ApplyButton is enabled */
-	bool IsApplyButtonEnabled() const;
+	bool IsPrefixRemoveSeparatorEnabled() const;
+	bool OnPrefixSeparatorVerifyTextChanged(const FText& InText, FText& OutErrorText) const;
+	void OnPrefixSeparatorChanged(const FText& InNewText);
 
-	/** Called when clicking on the Apply button */
-	FReply OnApplyButtonClicked();
+	ECheckBoxState IsPrefixRemoveCharactersChecked() const;
+	void OnPrefixRemoveCharactersCheckBoxChanged(ECheckBoxState InNewState);
 
-	/** Called when clicking on the Reset button */
-	FReply OnResetButtonClicked();
+	bool IsPrefixRemoveNumCharsEnabled() const;
+	void OnPrefixRemoveCharactersChanged(uint8 InNewValue);
 
-	/** Called when clicking on the Cancel button */
-	FReply OnCancelButtonClicked();
+	void OnSuffixChanged(const FText& InNewText);
 
-private:
-	/** Min update frequency, used in the tick to avoid updating the renamer too often */
-	static constexpr double MinUpdateFrequency = 0.1f;
+	ECheckBoxState IsSuffixRemoveChecked() const;
+	void OnSuffixRemoveCheckBoxChanged(ECheckBoxState InNewState);
 
-	/** Renamer instance */
-	TSharedPtr<IAdvancedRenamer> Renamer;
+	bool IsSuffixRemoveSeparatorEnabled() const;
+	bool OnSuffixSeparatorVerifyTextChanged(const FText& InText, FText& OutErrorText) const;
+	void OnSuffixSeparatorChanged(const FText& InNewText);
 
-	/** Command list of the Renamer Panel */
-	TSharedPtr<FUICommandList> CommandList;
+	ECheckBoxState IsSuffixRemoveCharactersChecked() const;
+	void OnSuffixRemoveCharactersCheckBoxChanged(ECheckBoxState InNewState);
 
-	/** Last update time of the list */
-	double ListLastUpdateTime = 0;
+	bool IsSuffixRemoveNumCharsEnabled() const;
+	void OnSuffixRemoveCharactersChanged(uint8 InNewValue);
 
-	/** Sections VerticalBox */
-	TSharedPtr<SVerticalBox> LeftSideVerticalBox;
+	ECheckBoxState IsSuffixRemoveNumberChecked() const;
+	void OnSuffixRemoveNumberCheckBoxChanged(ECheckBoxState InNewState);
 
-	/** RightSide container Box */
-	TSharedPtr<SBox> RightSideBox;
+	ECheckBoxState IsSuffixNumberChecked() const;
+	bool IsSuffixRemoveNumberCheckBoxEnabled() const;
+	void OnSuffixNumberCheckBoxChanged(ECheckBoxState InNewState);
 
-	/** Bottom panel button container Box */
-	TSharedPtr<SBox> ApplyResetCancelBox;
+	void OnSuffixNumberStartChanged(int32 InNewValue);
 
-	/** ListView container Box */
-	TSharedPtr<SBox> RenamePreviewListBox;
+	void OnSuffixNumberStepChanged(int32 InNewValue);
 
-	/** ListView Header row */
-	TSharedPtr<SHeaderRow> RenamePreviewListHeaderRow;
+	ECheckBoxState IsSearchReplacePlainTextChecked() const;
+	void OnSearchReplacePlainTextCheckBoxChanged(ECheckBoxState InNewState);
 
-	/** ListView of the Previews */
-	TSharedPtr<SListView<TSharedPtr<FAdvancedRenamerPreview>>> RenamePreviewList;
+	ECheckBoxState IsSearchReplaceRegexChecked() const;
+	void OnSearchReplaceRegexCheckBoxChanged(ECheckBoxState InNewState);
+
+	ECheckBoxState IsSearchReplaceIgnoreCaseChecked() const;
+	void OnSearchReplaceIgnoreCaseCheckBoxChanged(ECheckBoxState InNewState);
+
+	void OnSearchReplaceSearchTextChanged(const FText& InNewText);
+
+	void OnSearchReplaceReplaceTextChanged(const FText& InNewText);
+
+	TSharedRef<ITableRow> OnGenerateRowForList(TSharedPtr<FAdvancedRenamerPreview> InItem, const TSharedRef<STableViewBase>& InOwnerTable);
+
+	FReply OnListViewKeyDown(const FGeometry& InGeometry, const FKeyEvent& InKeyEvent);
+
+	TSharedPtr<SWidget> GenerateListViewContextMenu();
+
+	bool IsApplyButtonEnabled() const;
+	FReply OnApplyButtonClicked();
+
+	FVector2D GetRightPaneSize() const;
+
+	FVector2D GetListViewsize() const;
+
+	FVector2D GetApplyButtonSize() const;
 };
diff --git a/Engine/Plugins/Experimental/AdvancedRenamer/Source/AdvancedRenamer/Private/Utils/AdvancedRenamerSlateUtils.h b/Engine/Plugins/Experimental/AdvancedRenamer/Source/AdvancedRenamer/Private/Utils/AdvancedRenamerSlateUtils.h
deleted file mode 100644
diff --git a/Engine/Plugins/Experimental/AdvancedRenamer/Source/AdvancedRenamer/Public/IAdvancedRenamerModule.h b/Engine/Plugins/Experimental/AdvancedRenamer/Source/AdvancedRenamer/Public/IAdvancedRenamerModule.h
--- a/Engine/Plugins/Experimental/AdvancedRenamer/Source/AdvancedRenamer/Public/IAdvancedRenamerModule.h
+++ b/Engine/Plugins/Experimental/AdvancedRenamer/Source/AdvancedRenamer/Public/IAdvancedRenamerModule.h
@@ -9,7 +9,6 @@
 class AActor;
 class IAdvancedRenamer;
 class IAdvancedRenamerProvider;
-class IAdvancedRenamerSection;
 class IToolkitHost;
 class SWidget;
 
@@ -44,7 +43,5 @@ public:
 
 	virtual void OpenAdvancedRenamer(const TSharedRef<IAdvancedRenamer>& InRenamer, const TSharedPtr<SWidget>& InParentWidget) = 0;
 
-	virtual TArray<TSharedPtr<IAdvancedRenamerSection>> GetRegisteredSections() const = 0;
-
 	virtual TArray<AActor*> GetActorsSharingClassesInWorld(const TArray<AActor*>& InActors) = 0;
 };
