[BatchRenamer] Update plugin name and description
Update the plugin UI and set it to EnabledByDefault

diff --git a/Engine/Plugins/Experimental/AdvancedRenamer/AdvancedRenamer.uplugin b/Engine/Plugins/Experimental/AdvancedRenamer/AdvancedRenamer.uplugin
--- a/Engine/Plugins/Experimental/AdvancedRenamer/AdvancedRenamer.uplugin
+++ b/Engine/Plugins/Experimental/AdvancedRenamer/AdvancedRenamer.uplugin
@@ -2,8 +2,8 @@
 	"FileVersion": 3,
 	"Version": 1,
 	"VersionName": "1.0",
-	"FriendlyName": "Advanced Renamer",
-	"Description": "Advanced Renaming tool for actors and assets",
+	"FriendlyName": "Batch Renamer",
+	"Description": "Rename multiple selected actors or assets, and standardize their prefixes and suffixes.",
 	"Category": "Experimental",
 	"CreatedBy": "Epic Games, Inc.",
 	"CreatedByURL": "https://epicgames.com",
@@ -13,6 +13,7 @@
 	"CanContainContent": false,
 	"IsBetaVersion": false,
 	"IsExperimentalVersion": true,
+	"EnabledByDefault": true,
 	"Installed": false,
 	"Modules": [
 		{
diff --git a/Engine/Plugins/Experimental/AdvancedRenamer/Source/AdvancedRenamer/Private/AdvancedRenamerCommands.cpp b/Engine/Plugins/Experimental/AdvancedRenamer/Source/AdvancedRenamer/Private/AdvancedRenamerCommands.cpp
--- a/Engine/Plugins/Experimental/AdvancedRenamer/Source/AdvancedRenamer/Private/AdvancedRenamerCommands.cpp
+++ b/Engine/Plugins/Experimental/AdvancedRenamer/Source/AdvancedRenamer/Private/AdvancedRenamerCommands.cpp
@@ -22,13 +22,13 @@ void FAdvancedRenamerCommands::RegisterCommands()
 {
 	UI_COMMAND(RenameSelectedActors
 		, "Rename Selected Actors"
-		, "Opens the Advanced Renamer Panel to rename all selected actors."
+		, "Opens the Batch Renamer Panel to rename all selected actors."
 		, EUserInterfaceActionType::Button
 		, FInputChord())
 
 	UI_COMMAND(RenameSharedClassActors
 		, "Rename Actors of Selected Actor Classes"
-		, "Opens the Advanced Renamer Panel to rename all actors sharing a class with any selected actor."
+		, "Opens the Batch Renamer Panel to rename all actors sharing a class with any selected actor."
 		, EUserInterfaceActionType::Button
 		, FInputChord())
 }
diff --git a/Engine/Plugins/Experimental/AdvancedRenamer/Source/AdvancedRenamer/Private/AdvancedRenamerModule.cpp b/Engine/Plugins/Experimental/AdvancedRenamer/Source/AdvancedRenamer/Private/AdvancedRenamerModule.cpp
--- a/Engine/Plugins/Experimental/AdvancedRenamer/Source/AdvancedRenamer/Private/AdvancedRenamerModule.cpp
+++ b/Engine/Plugins/Experimental/AdvancedRenamer/Source/AdvancedRenamer/Private/AdvancedRenamerModule.cpp
@@ -28,7 +28,7 @@ namespace UE::AdvancedRenamer::Private
 	TSharedRef<SWindow> CreateAdvancedRenamerWindow()
 	{
 		return SNew(SWindow)
-			.Title(LOCTEXT("AdvancedRenameWindow", "Rename Actors"))
+			.Title(LOCTEXT("AdvancedRenameWindow", "Batch Renamer"))
 			.ClientSize(FVector2D(731, 547))
 			.SizingRule(ESizingRule::UserSized)
 			.SupportsMaximize(false)
diff --git a/Engine/Plugins/Experimental/AdvancedRenamer/Source/AdvancedRenamer/Private/Integrations/AdvancedRenamerContentBrowserIntegration.cpp b/Engine/Plugins/Experimental/AdvancedRenamer/Source/AdvancedRenamer/Private/Integrations/AdvancedRenamerContentBrowserIntegration.cpp
--- a/Engine/Plugins/Experimental/AdvancedRenamer/Source/AdvancedRenamer/Private/Integrations/AdvancedRenamerContentBrowserIntegration.cpp
+++ b/Engine/Plugins/Experimental/AdvancedRenamer/Source/AdvancedRenamer/Private/Integrations/AdvancedRenamerContentBrowserIntegration.cpp
@@ -29,8 +29,8 @@ namespace UE::AdvancedRenamer::Private
 	void ExtendAssetMenu(FMenuBuilder& MenuBuilder, const TArray<FAssetData> AssetArray)
 	{
 		MenuBuilder.AddMenuEntry(
-			LOCTEXT("AdvancedRename", "Advanced Rename"),
-			LOCTEXT("AdvancedRenameTooltip", "Opens the Advanced Renamer Panel to rename all selected assets."),
+			LOCTEXT("AdvancedRename", "Batch Rename"),
+			LOCTEXT("AdvancedRenameTooltip", "Opens the Batch Renamer Panel to rename all selected assets."),
 			FSlateIcon(FAppStyle::GetAppStyleSetName(), "ContentBrowser.AssetActions.Rename"),
 			FUIAction(FExecuteAction::CreateStatic(&OpenAdvancedRenamer, AssetArray))
 		);
diff --git a/Engine/Plugins/Experimental/AdvancedRenamer/Source/AdvancedRenamer/Private/Slate/SAdvancedRenamerPanel.cpp b/Engine/Plugins/Experimental/AdvancedRenamer/Source/AdvancedRenamer/Private/Slate/SAdvancedRenamerPanel.cpp
--- a/Engine/Plugins/Experimental/AdvancedRenamer/Source/AdvancedRenamer/Private/Slate/SAdvancedRenamerPanel.cpp
+++ b/Engine/Plugins/Experimental/AdvancedRenamer/Source/AdvancedRenamer/Private/Slate/SAdvancedRenamerPanel.cpp
@@ -489,7 +489,7 @@ bool SAdvancedRenamerPanel::IsApplyButtonEnabled() const
 
 FReply SAdvancedRenamerPanel::OnApplyButtonClicked()
 {
-	FScopedTransaction Transaction(LOCTEXT("AdvancedRenamerRename", "Advanced Renamer Rename"));
+	FScopedTransaction Transaction(LOCTEXT("AdvancedRenamerRename", "Batch Renamer Rename"));
 
 	if (!Renamer->Execute())
 	{
