[BatchRenamer] UE5/Main: Issue 464491
Text conflict from LOCTEXT macro

diff --git a/Engine/Plugins/Experimental/AdvancedRenamer/Source/AdvancedRenamer/Private/AdvancedRenamerSections/AdvancedRenamerRemovePrefixSection.cpp b/Engine/Plugins/Experimental/AdvancedRenamer/Source/AdvancedRenamer/Private/AdvancedRenamerSections/AdvancedRenamerRemovePrefixSection.cpp
--- a/Engine/Plugins/Experimental/AdvancedRenamer/Source/AdvancedRenamer/Private/AdvancedRenamerSections/AdvancedRenamerRemovePrefixSection.cpp
+++ b/Engine/Plugins/Experimental/AdvancedRenamer/Source/AdvancedRenamer/Private/AdvancedRenamerSections/AdvancedRenamerRemovePrefixSection.cpp
@@ -93,10 +93,10 @@ TSharedRef<SWidget> FAdvancedRenamerRemovePrefixSection::GetWidget()
 						.OnValueChanged(this, &FAdvancedRenamerRemovePrefixSection::OnRemoveOldPrefixTypeChanged)
 
 						+ SSegmentedControl<EAdvancedRenamerRemoveOldType>::Slot(EAdvancedRenamerRemoveOldType::Separator)
-						.Text(LOCTEXT("AR_SuffixSeparator", "Separator"))
+						.Text(LOCTEXT("AR_PrefixSeparator", "Separator"))
 
 						+ SSegmentedControl<EAdvancedRenamerRemoveOldType>::Slot(EAdvancedRenamerRemoveOldType::Chars)
-						.Text(LOCTEXT("AR_SuffixFirstChar(s)", "First Char(s)"))
+						.Text(LOCTEXT("AR_PrefixFirstChar(s)", "First Char(s)"))
 					]
 				]
 				+ SHorizontalBox::Slot()
@@ -136,7 +136,7 @@ void FAdvancedRenamerRemovePrefixSection::ResetToDefault()
 	bRemoveOldPrefixSection = false;
 	PrefixWidgetSwitcherIndex = 0;
 	RemovePrefixCharsValue = 1;
-	RemovePrefixSeparatorText = LOCTEXT("AR_PrefixSeparator", "_");
+	RemovePrefixSeparatorText = LOCTEXT("AR_PrefixSeparatorText", "_");
 	RemovePrefixType = EAdvancedRenamerRemoveOldType::Separator;
 }
 
diff --git a/Engine/Plugins/Experimental/AdvancedRenamer/Source/AdvancedRenamer/Private/AdvancedRenamerSections/AdvancedRenamerRemoveSuffixSection.cpp b/Engine/Plugins/Experimental/AdvancedRenamer/Source/AdvancedRenamer/Private/AdvancedRenamerSections/AdvancedRenamerRemoveSuffixSection.cpp
--- a/Engine/Plugins/Experimental/AdvancedRenamer/Source/AdvancedRenamer/Private/AdvancedRenamerSections/AdvancedRenamerRemoveSuffixSection.cpp
+++ b/Engine/Plugins/Experimental/AdvancedRenamer/Source/AdvancedRenamer/Private/AdvancedRenamerSections/AdvancedRenamerRemoveSuffixSection.cpp
@@ -167,7 +167,7 @@ void FAdvancedRenamerRemoveSuffixSection::ResetToDefault()
 	bRemoveSuffixNumbers = false;
 	SuffixWidgetSwitcherIndex = 0;
 	RemoveSuffixCharsValue = 1;
-	RemoveSuffixSeparatorText = LOCTEXT("AR_SuffixSeparator", "_");
+	RemoveSuffixSeparatorText = LOCTEXT("AR_SuffixSeparatorText", "_");
 	RemoveSuffixType = EAdvancedRenamerRemoveOldType::Separator;
 }
 
