AdvancedRenamer: Changed the content browser menu extension to use UToolMenus so that its compatible with permissions filtering.

diff --git a/Engine/Plugins/Experimental/AdvancedRenamer/Source/AdvancedRenamer/Private/Integrations/AdvancedRenamerContentBrowserIntegration.cpp b/Engine/Plugins/Experimental/AdvancedRenamer/Source/AdvancedRenamer/Private/Integrations/AdvancedRenamerContentBrowserIntegration.cpp
--- a/Engine/Plugins/Experimental/AdvancedRenamer/Source/AdvancedRenamer/Private/Integrations/AdvancedRenamerContentBrowserIntegration.cpp
+++ b/Engine/Plugins/Experimental/AdvancedRenamer/Source/AdvancedRenamer/Private/Integrations/AdvancedRenamerContentBrowserIntegration.cpp
@@ -2,6 +2,7 @@
 
 #include "Integrations/AdvancedRenamerContentBrowserIntegration.h"
 #include "ContentBrowserDelegates.h"
+#include "ContentBrowserMenuContexts.h"
 #include "ContentBrowserModule.h"
 #include "Delegates/IDelegateInstance.h"
 #include "Framework/Commands/UIAction.h"
@@ -9,6 +10,9 @@
 #include "Framework/MultiBox/MultiBoxExtender.h"
 #include "IAdvancedRenamerModule.h"
 #include "Providers/AdvancedRenamerAssetProvider.h"
+#include "ToolMenus.h"
+#include "ToolMenuSection.h"
+
 
 #define LOCTEXT_NAMESPACE "AdvancedRenamerContentBrowserIntegration"
 
@@ -26,31 +30,41 @@ namespace UE::AdvancedRenamer::Private
 		IAdvancedRenamerModule::Get().OpenAdvancedRenamer(StaticCastSharedRef<IAdvancedRenamerProvider>(AssetProvider), HostWidget);
 	}
 
-	void ExtendAssetMenu(FMenuBuilder& MenuBuilder, const TArray<FAssetData> AssetArray)
+	void AddMenuEntry(FToolMenuSection& MenuSection, const TArray<FAssetData> SelectedAssets)
 	{
-		MenuBuilder.AddMenuEntry(
+		MenuSection.AddMenuEntry(
+			"BatchRename",
 			LOCTEXT("AdvancedRename", "Batch Rename"),
 			LOCTEXT("AdvancedRenameTooltip", "Opens the Batch Renamer Panel to rename all selected assets."),
 			FSlateIcon(FAppStyle::GetAppStyleSetName(), "ContentBrowser.AssetActions.Rename"),
-			FUIAction(FExecuteAction::CreateStatic(&OpenAdvancedRenamer, AssetArray))
-		);
+			FUIAction(FExecuteAction::CreateStatic(&OpenAdvancedRenamer, SelectedAssets)));
 	}
 
-	TSharedRef<FExtender> CreateContentMenuExtender(const TArray<FAssetData>& AssetArray)
+	void ExtendAssetContextMenu()
 	{
-		TSharedRef<FExtender> Extender = MakeShared<FExtender>();
-
-		if (AssetArray.Num() > 0)
+		if (UToolMenus* ToolMenus = UToolMenus::Get())
 		{
-			Extender->AddMenuExtension(
-				"CommonAssetActions",
-				EExtensionHook::Position::After,
-				nullptr,
-				FMenuExtensionDelegate::CreateStatic(&ExtendAssetMenu, AssetArray)
-			);
-		}
+			FToolMenuOwnerScoped OwnerScoped(TEXT("AdvancedRenamer"));
 
-		return Extender;
+			if (UToolMenu* Menu = ToolMenus->ExtendMenu("ContentBrowser.AssetContextMenu"))
+			{
+				if (FToolMenuSection* MenuSection = Menu->FindSection("CommonAssetActions"))
+				{
+					MenuSection->AddDynamicEntry("CreateVariant", FNewToolMenuSectionDelegate::CreateLambda([](FToolMenuSection& InSection)
+					{
+						UContentBrowserAssetContextMenuContext* Context = InSection.FindContext<UContentBrowserAssetContextMenuContext>();
+						if (Context)
+						{
+							const TArray<FAssetData>& SelectedAssets = Context->SelectedAssets;
+							if (SelectedAssets.Num() > 0 )
+							{
+								AddMenuEntry(InSection, SelectedAssets);
+							}
+						}
+					}));
+				}
+			}
+		}
 	}
 }
 
@@ -59,30 +73,14 @@ void FAdvancedRenamerContentBrowserIntegration::Initialize()
 	using namespace UE::AdvancedRenamer::Private;
 
 	// Register Content Browser selection extensions
-	FContentBrowserModule& ContentBrowserModule = FModuleManager::LoadModuleChecked<FContentBrowserModule>("ContentBrowser");
-	TArray<FContentBrowserMenuExtender_SelectedAssets>& AssetContextMenuDelegates = ContentBrowserModule.GetAllAssetViewContextMenuExtenders();
-	AssetContextMenuDelegates.Add(FContentBrowserMenuExtender_SelectedAssets::CreateStatic(&CreateContentMenuExtender));
-	ContentBrowserDelegateHandle = AssetContextMenuDelegates.Last().GetHandle();
+	ExtendAssetContextMenu();
 }
 
 void FAdvancedRenamerContentBrowserIntegration::Shutdown()
 {
-	using namespace UE::AdvancedRenamer::Private;
-
-	if (ContentBrowserDelegateHandle.IsValid())
+	if (UToolMenus* ToolMenus = UToolMenus::TryGet())
 	{
-		FContentBrowserModule* ContentBrowserModule = FModuleManager::GetModulePtr<FContentBrowserModule>("ContentBrowser");
-
-		if (ContentBrowserModule)
-		{
-			TArray<FContentBrowserMenuExtender_SelectedAssets>& AssetContextMenuDelegates = ContentBrowserModule->GetAllAssetViewContextMenuExtenders();
-			AssetContextMenuDelegates.RemoveAll([DelegateHandle = ContentBrowserDelegateHandle](const FContentBrowserMenuExtender_SelectedAssets& Delegate)
-				{
-					return Delegate.GetHandle() == DelegateHandle;
-				});
-		}
-
-		ContentBrowserDelegateHandle.Reset();
+		ToolMenus->UnregisterOwner(TEXT("AdvancedRenamer"));
 	}
 }
 
