[BatchRenamer] Remove the CVar to enable the integrations with the CB and the LevelEditor by default and add commands to the menus
Removed CVar to hide the plugin.
Adding Commands to use for the ContentBrowser and Outliner menu.
Added ability to rebind commands.

diff --git a/Engine/Plugins/Experimental/AdvancedRenamer/Source/AdvancedRenamer/Private/AdvancedRenamerCommands.cpp b/Engine/Plugins/Experimental/AdvancedRenamer/Source/AdvancedRenamer/Private/AdvancedRenamerCommands.cpp
--- a/Engine/Plugins/Experimental/AdvancedRenamer/Source/AdvancedRenamer/Private/AdvancedRenamerCommands.cpp
+++ b/Engine/Plugins/Experimental/AdvancedRenamer/Source/AdvancedRenamer/Private/AdvancedRenamerCommands.cpp
@@ -20,17 +20,8 @@ FAdvancedRenamerCommands::FAdvancedRenamerCommands()
 
 void FAdvancedRenamerCommands::RegisterCommands()
 {
-	UI_COMMAND(RenameSelectedActors
-		, "Rename Selected Actors"
-		, "Opens the Batch Renamer Panel to rename all selected actors."
-		, EUserInterfaceActionType::Button
-		, FInputChord())
-
-	UI_COMMAND(RenameSharedClassActors
-		, "Rename Actors of Selected Actor Classes"
-		, "Opens the Batch Renamer Panel to rename all actors sharing a class with any selected actor."
-		, EUserInterfaceActionType::Button
-		, FInputChord())
+	UI_COMMAND(BatchRenameObject, "Batch Rename", "Batch Rename Object(s) based on selection.", EUserInterfaceActionType::Button, FInputChord(EModifierKey::Shift, EKeys::F2));
+	UI_COMMAND(BatchRenameSharedClassActors, "Rename Actors of Selected Actor Classes", "Opens the Batch Renamer Panel to rename all actors sharing a class with any selected actor.", EUserInterfaceActionType::Button, FInputChord(EModifierKey::Control | EModifierKey::Shift, EKeys::F2));
 }
 
 #undef LOCTEXT_NAMESPACE
diff --git a/Engine/Plugins/Experimental/AdvancedRenamer/Source/AdvancedRenamer/Private/AdvancedRenamerCommands.h b/Engine/Plugins/Experimental/AdvancedRenamer/Source/AdvancedRenamer/Private/AdvancedRenamerCommands.h
--- a/Engine/Plugins/Experimental/AdvancedRenamer/Source/AdvancedRenamer/Private/AdvancedRenamerCommands.h
+++ b/Engine/Plugins/Experimental/AdvancedRenamer/Source/AdvancedRenamer/Private/AdvancedRenamerCommands.h
@@ -14,6 +14,6 @@ public:
 
 	virtual void RegisterCommands() override;
 
-	TSharedPtr<FUICommandInfo> RenameSelectedActors;
-	TSharedPtr<FUICommandInfo> RenameSharedClassActors;
+	TSharedPtr<FUICommandInfo> BatchRenameObject;
+	TSharedPtr<FUICommandInfo> BatchRenameSharedClassActors;
 };
diff --git a/Engine/Plugins/Experimental/AdvancedRenamer/Source/AdvancedRenamer/Private/AdvancedRenamerModule.cpp b/Engine/Plugins/Experimental/AdvancedRenamer/Source/AdvancedRenamer/Private/AdvancedRenamerModule.cpp
--- a/Engine/Plugins/Experimental/AdvancedRenamer/Source/AdvancedRenamer/Private/AdvancedRenamerModule.cpp
+++ b/Engine/Plugins/Experimental/AdvancedRenamer/Source/AdvancedRenamer/Private/AdvancedRenamerModule.cpp
@@ -36,65 +36,23 @@ namespace UE::AdvancedRenamer::Private
 			.MinWidth(737)
 			.MinHeight(586.5f);
 	}
-
-	static TAutoConsoleVariable<bool> CVarBatchRenamerEnableIntegration(
-		TEXT("BatchRenamer.EnableIntegration"),
-		false,
-		TEXT("If enabled, the Batch Renamer will hook into the ContentBrowser and Outliner.")
-	);
-
-	static bool bPreviousCVarValue = false;
-
-	void OnBatchRenamerEnabledChanged(IConsoleVariable* InVar)
-	{
-		if (InVar)
-		{
-			const bool bCurrentValue = InVar->GetBool();
-			if (bCurrentValue != bPreviousCVarValue)
-			{
-				if (InVar->GetBool())
-				{
-					FAdvancedRenamerContentBrowserIntegration::Initialize();
-					FAdvancedRenamerLevelEditorIntegration::InitializeMenu();
-				}
-				else
-				{
-					FAdvancedRenamerContentBrowserIntegration::Shutdown();
-					FAdvancedRenamerLevelEditorIntegration::ShutdownMenu();
-				}
-				bPreviousCVarValue = bCurrentValue;
-			}
-		}
-	}
 }
 
 void FAdvancedRenamerModule::StartupModule()
 {
 	FAdvancedRenamerStyle::Initialize();
-	FAdvancedRenamerLevelEditorIntegration::Initialize();
 	FAdvancedRenamerCommands::Register();
-
-	{
-		// Enable plugins integration section from CVar
-		using namespace UE::AdvancedRenamer::Private;
-		EnableRenamerHandle = CVarBatchRenamerEnableIntegration.AsVariable()->OnChangedDelegate().AddStatic(&OnBatchRenamerEnabledChanged);
-	}
-
+	FAdvancedRenamerContentBrowserIntegration::Initialize();
+	FAdvancedRenamerLevelEditorIntegration::Initialize();
 	RegisterDefaultSections();
 }
 
 void FAdvancedRenamerModule::ShutdownModule()
 {
-	FAdvancedRenamerCommands::Unregister();
 	FAdvancedRenamerStyle::Shutdown();
+	FAdvancedRenamerCommands::Unregister();
 	FAdvancedRenamerContentBrowserIntegration::Shutdown();
 	FAdvancedRenamerLevelEditorIntegration::Shutdown();
-	FAdvancedRenamerLevelEditorIntegration::ShutdownMenu();
-	if (EnableRenamerHandle.IsValid())
-	{
-		using namespace UE::AdvancedRenamer::Private;
-		CVarBatchRenamerEnableIntegration.AsVariable()->OnChangedDelegate().Remove(EnableRenamerHandle);
-	}
 }
 
 TSharedRef<IAdvancedRenamer> FAdvancedRenamerModule::CreateAdvancedRenamer(const TSharedRef<IAdvancedRenamerProvider>& InRenameProvider)
diff --git a/Engine/Plugins/Experimental/AdvancedRenamer/Source/AdvancedRenamer/Private/AdvancedRenamerStyle.cpp b/Engine/Plugins/Experimental/AdvancedRenamer/Source/AdvancedRenamer/Private/AdvancedRenamerStyle.cpp
--- a/Engine/Plugins/Experimental/AdvancedRenamer/Source/AdvancedRenamer/Private/AdvancedRenamerStyle.cpp
+++ b/Engine/Plugins/Experimental/AdvancedRenamer/Source/AdvancedRenamer/Private/AdvancedRenamerStyle.cpp
@@ -8,6 +8,8 @@
 #include "Styling/SlateStyleRegistry.h"
 #include "Styling/SlateTypes.h"
 
+#define IMAGE_BRUSH(RelativePath, ...) FSlateImageBrush(StyleInstance->RootToContentDir(RelativePath, TEXT(".png") ), __VA_ARGS__)
+
 TSharedPtr<FSlateStyleSet> FAdvancedRenamerStyle::StyleInstance = nullptr;
 
 void FAdvancedRenamerStyle::Initialize()
@@ -44,13 +46,9 @@ void FAdvancedRenamerStyle::InitStyle()
 
 	StyleInstance = MakeShared<FSlateStyleSet>("AdvancedRenamer");
 
-	TSharedPtr<IPlugin> Plugin = IPluginManager::Get().FindPlugin(TEXT("AdvancedRenamer"));
-	check(Plugin.IsValid());
-
-	if (Plugin.IsValid())
-	{
-		StyleInstance->SetContentRoot(FPaths::Combine(Plugin->GetBaseDir(), TEXT("Resources")));
-	}
+	// Same ContentDir and CoreRootContentDir as the ContentBrowser
+	StyleInstance->SetContentRoot(FPaths::EngineContentDir() / TEXT("Editor/Slate"));
+	StyleInstance->SetCoreContentRoot(FPaths::EngineContentDir() / TEXT("Slate"));
 
 	const FSplitterStyle SplitterStyle = FSplitterStyle()
 		.SetHandleNormalBrush(FSlateNoResource())
@@ -76,6 +74,10 @@ void FAdvancedRenamerStyle::InitStyle()
 
 	StyleInstance->Set("AdvancedRenamer.Style.TitleFont", FCoreStyle::GetDefaultFontStyle("Regular", 12));
 	StyleInstance->Set("AdvancedRenamer.Style.RegularFont", FCoreStyle::GetDefaultFontStyle("Regular", 10));
+
+	// Commands Icon
+	StyleInstance->Set("AdvancedRenamer.BatchRenameObject", new IMAGE_BRUSH("Icons/Icon_Asset_Rename_16x", FVector2D(16.f, 16.f)));
+	StyleInstance->Set("AdvancedRenamer.BatchRenameSharedClassActors", new IMAGE_BRUSH("Icons/Icon_Asset_Rename_16x", FVector2D(16.f, 16.f)));
 }
 
 const ISlateStyle& FAdvancedRenamerStyle::Get()
@@ -87,3 +89,5 @@ const ISlateStyle& FAdvancedRenamerStyle::Get()
 
 	return *StyleInstance;
 }
+
+#undef IMAGE_BRUSH
diff --git a/Engine/Plugins/Experimental/AdvancedRenamer/Source/AdvancedRenamer/Private/Integrations/AdvancedRenamerContentBrowserIntegration.cpp b/Engine/Plugins/Experimental/AdvancedRenamer/Source/AdvancedRenamer/Private/Integrations/AdvancedRenamerContentBrowserIntegration.cpp
--- a/Engine/Plugins/Experimental/AdvancedRenamer/Source/AdvancedRenamer/Private/Integrations/AdvancedRenamerContentBrowserIntegration.cpp
+++ b/Engine/Plugins/Experimental/AdvancedRenamer/Source/AdvancedRenamer/Private/Integrations/AdvancedRenamerContentBrowserIntegration.cpp
@@ -1,11 +1,14 @@
 // Copyright Epic Games, Inc. All Rights Reserved.
 
 #include "Integrations/AdvancedRenamerContentBrowserIntegration.h"
+
+#include "AdvancedRenamerCommands.h"
 #include "ContentBrowserDelegates.h"
 #include "ContentBrowserMenuContexts.h"
 #include "ContentBrowserModule.h"
 #include "Delegates/IDelegateInstance.h"
 #include "Framework/Commands/UIAction.h"
+#include "Framework/Commands/UICommandList.h"
 #include "Framework/MultiBox/MultiBoxBuilder.h"
 #include "Framework/MultiBox/MultiBoxExtender.h"
 #include "IAdvancedRenamerModule.h"
@@ -13,74 +16,110 @@
 #include "ToolMenus.h"
 #include "ToolMenuSection.h"
 
-
 #define LOCTEXT_NAMESPACE "AdvancedRenamerContentBrowserIntegration"
 
 namespace UE::AdvancedRenamer::Private
 {
 	FDelegateHandle ContentBrowserDelegateHandle;
 
-	void OpenAdvancedRenamer(const TArray<FAssetData> AssetArray)
+	TArray<FAssetData> GetContentBrowserSelectedAssets(FOnContentBrowserGetSelection InGetSelectionDelegate)
 	{
-		TSharedRef<FAdvancedRenamerAssetProvider> AssetProvider = MakeShared<FAdvancedRenamerAssetProvider>();
-		AssetProvider->SetAssetList(AssetArray);
+		if (InGetSelectionDelegate.IsBound())
+		{
+			TArray<FAssetData> SelectedAssets;
+			TArray<FString> SelectedPaths;
+			InGetSelectionDelegate.Execute(SelectedAssets, SelectedPaths);
+			return SelectedAssets;
+		}
+		return TArray<FAssetData>();
+	}
+
+	void OpenAdvancedRenamer(FOnContentBrowserGetSelection InGetSelectionDelegate)
+	{
+		const TArray<FAssetData> SelectedAssets = GetContentBrowserSelectedAssets(InGetSelectionDelegate);
+		if (SelectedAssets.IsEmpty())
+		{
+			return;
+		}
+
+		const TSharedRef<FAdvancedRenamerAssetProvider> AssetProvider = MakeShared<FAdvancedRenamerAssetProvider>();
+		AssetProvider->SetAssetList(SelectedAssets);
 
 		TSharedPtr<SWidget> HostWidget = nullptr;
 
 		IAdvancedRenamerModule::Get().OpenAdvancedRenamer(StaticCastSharedRef<IAdvancedRenamerProvider>(AssetProvider), HostWidget);
 	}
 
-	void AddMenuEntry(FToolMenuSection& MenuSection, const TArray<FAssetData> SelectedAssets)
+	void AddMenuEntry(FToolMenuSection& InMenuSection)
 	{
-		MenuSection.AddMenuEntry(
-			"BatchRename",
-			LOCTEXT("AdvancedRename", "Batch Rename"),
-			LOCTEXT("AdvancedRenameTooltip", "Opens the Batch Renamer Panel to rename all selected assets."),
-			FSlateIcon(FAppStyle::GetAppStyleSetName(), "ContentBrowser.AssetActions.Rename"),
-			FUIAction(FExecuteAction::CreateStatic(&OpenAdvancedRenamer, SelectedAssets)));
+		FToolMenuEntry& BatchRenameMenuEntry = InMenuSection.AddMenuEntry(
+			FAdvancedRenamerCommands::Get().BatchRenameObject,
+			LOCTEXT("BatchRename", "Batch Rename"),
+			LOCTEXT("AdvancedRenameTooltip", "Opens the Batch Renamer Panel to rename all selected assets."));
+
+		BatchRenameMenuEntry.InsertPosition = FToolMenuInsert(TEXT("Rename"), EToolMenuInsertType::After);
 	}
 
-	void ExtendAssetContextMenu()
+	void RegisterAssetMenu()
 	{
 		if (UToolMenus* ToolMenus = UToolMenus::Get())
 		{
 			FToolMenuOwnerScoped OwnerScoped(TEXT("AdvancedRenamer"));
 
-			if (UToolMenu* Menu = ToolMenus->ExtendMenu("ContentBrowser.AssetContextMenu"))
+			if (UToolMenu* Menu = ToolMenus->ExtendMenu(TEXT("ContentBrowser.AssetContextMenu")))
 			{
-				if (FToolMenuSection* MenuSection = Menu->FindSection("CommonAssetActions"))
-				{
-					MenuSection->AddDynamicEntry("CreateVariant", FNewToolMenuSectionDelegate::CreateLambda([](FToolMenuSection& InSection)
-					{
-						UContentBrowserAssetContextMenuContext* Context = InSection.FindContext<UContentBrowserAssetContextMenuContext>();
-						if (Context)
-						{
-							const TArray<FAssetData>& SelectedAssets = Context->SelectedAssets;
-							if (SelectedAssets.Num() > 0 )
-							{
-								AddMenuEntry(InSection, SelectedAssets);
-							}
-						}
-					}));
-				}
+				FToolMenuSection& MenuSection = Menu->FindOrAddSection(TEXT("CommonAssetActions"));
+				AddMenuEntry(MenuSection);
 			}
 		}
 	}
+
+	void OnExtendContentBrowserCommands(TSharedRef<FUICommandList> OutCommandList, FOnContentBrowserGetSelection InGetSelectionDelegate)
+	{
+		using namespace UE::AdvancedRenamer::Private;
+
+		OutCommandList->MapAction(
+			FAdvancedRenamerCommands::Get().BatchRenameObject,
+			FUIAction(FExecuteAction::CreateStatic(&OpenAdvancedRenamer, InGetSelectionDelegate)));
+
+		FInputBindingManager::Get().RegisterCommandList(FAdvancedRenamerCommands::Get().GetContextName(), OutCommandList);
+	}
+
+	void ExtendContentBrowserCommands()
+	{
+		if (FModuleManager::Get().IsModuleLoaded(TEXT("ContentBrowser")))
+		{
+			FContentBrowserModule& ContentBrowserModule = FModuleManager::LoadModuleChecked<FContentBrowserModule>(TEXT("ContentBrowser"));
+			TArray<FContentBrowserCommandExtender>& CBCommandExtenderDelegates = ContentBrowserModule.GetAllContentBrowserCommandExtenders();
+			CBCommandExtenderDelegates.Add(FContentBrowserCommandExtender::CreateStatic(&OnExtendContentBrowserCommands));
+			ContentBrowserDelegateHandle = CBCommandExtenderDelegates.Last().GetHandle();
+		}
+	}
 }
 
 void FAdvancedRenamerContentBrowserIntegration::Initialize()
 {
 	using namespace UE::AdvancedRenamer::Private;
 
-	// Register Content Browser selection extensions
-	ExtendAssetContextMenu();
+	// Extend the Content Browser Commands
+	ExtendContentBrowserCommands();
+
+	// Register the menu entry
+	RegisterAssetMenu();
 }
 
 void FAdvancedRenamerContentBrowserIntegration::Shutdown()
 {
-	if (UToolMenus* ToolMenus = UToolMenus::TryGet())
+	using namespace UE::AdvancedRenamer::Private;
+
+	UToolMenus::UnregisterOwner(TEXT("AdvancedRenamer"));
+
+	if (FModuleManager::Get().IsModuleLoaded(TEXT("ContentBrowser")))
 	{
-		ToolMenus->UnregisterOwner(TEXT("AdvancedRenamer"));
+		FContentBrowserModule& ContentBrowserModule = FModuleManager::LoadModuleChecked<FContentBrowserModule>(TEXT("ContentBrowser"));
+		
+		TArray<FContentBrowserCommandExtender>& CBCommandExtenderDelegates = ContentBrowserModule.GetAllContentBrowserCommandExtenders();
+		CBCommandExtenderDelegates.RemoveAll([](const FContentBrowserCommandExtender& Delegate) { return Delegate.GetHandle() == ContentBrowserDelegateHandle; });
 	}
 }
 
diff --git a/Engine/Plugins/Experimental/AdvancedRenamer/Source/AdvancedRenamer/Private/Integrations/AdvancedRenamerLevelEditorIntegration.cpp b/Engine/Plugins/Experimental/AdvancedRenamer/Source/AdvancedRenamer/Private/Integrations/AdvancedRenamerLevelEditorIntegration.cpp
--- a/Engine/Plugins/Experimental/AdvancedRenamer/Source/AdvancedRenamer/Private/Integrations/AdvancedRenamerLevelEditorIntegration.cpp
+++ b/Engine/Plugins/Experimental/AdvancedRenamer/Source/AdvancedRenamer/Private/Integrations/AdvancedRenamerLevelEditorIntegration.cpp
@@ -13,6 +13,8 @@
 #include "Templates/SharedPointer.h"
 #include "ToolMenus.h"
 
+#define LOCTEXT_NAMESPACE "AdvancedRenamerLevelEditorIntegration"
+
 namespace UE::AdvancedRenamer::Private
 {
 	FDelegateHandle LevelEditorCreatedDelegateHandle;
@@ -89,17 +91,23 @@ namespace UE::AdvancedRenamer::Private
 	{
 		const FAdvancedRenamerCommands& AdvRenCommands = FAdvancedRenamerCommands::Get();
 
-		InLevelEditor->GetLevelEditorActions()->MapAction(
-			AdvRenCommands.RenameSelectedActors,
-			FExecuteAction::CreateStatic(&RenameSelectedActors, InLevelEditor.ToWeakPtr()),
-			FCanExecuteAction::CreateStatic(&CanOpenAdvancedRenamer, InLevelEditor.ToWeakPtr())
-		);
+		const TSharedPtr<FUICommandList>& LevelEditorAction = InLevelEditor->GetLevelEditorActions();
+		if (LevelEditorAction.IsValid())
+		{
+			LevelEditorAction->MapAction(
+				AdvRenCommands.BatchRenameObject,
+				FExecuteAction::CreateStatic(&RenameSelectedActors, InLevelEditor.ToWeakPtr()),
+				FCanExecuteAction::CreateStatic(&CanOpenAdvancedRenamer, InLevelEditor.ToWeakPtr())
+			);
+
+			LevelEditorAction->MapAction(
+				AdvRenCommands.BatchRenameSharedClassActors,
+				FExecuteAction::CreateStatic(&RenameSharedClassActors, InLevelEditor.ToWeakPtr()),
+				FCanExecuteAction::CreateStatic(&CanOpenAdvancedRenamer, InLevelEditor.ToWeakPtr())
+			);
+		}
 
-		InLevelEditor->GetLevelEditorActions()->MapAction(
-			AdvRenCommands.RenameSharedClassActors,
-			FExecuteAction::CreateStatic(&RenameSharedClassActors, InLevelEditor.ToWeakPtr()),
-			FCanExecuteAction::CreateStatic(&CanOpenAdvancedRenamer, InLevelEditor.ToWeakPtr())
-		);
+		FInputBindingManager::Get().RegisterCommandList(AdvRenCommands.GetContextName(), LevelEditorAction.ToSharedRef());
 	}
 }
 
@@ -114,6 +122,19 @@ void FAdvancedRenamerLevelEditorIntegration::Initialize()
 
 	FLevelEditorModule& LevelEditorModule = FModuleManager::LoadModuleChecked<FLevelEditorModule>("LevelEditor");
 	LevelEditorCreatedDelegateHandle = LevelEditorModule.OnLevelEditorCreated().AddStatic(&OnLevelEditorCreated);
+	InitializeMenu();
+}
+
+void FAdvancedRenamerLevelEditorIntegration::Shutdown()
+{
+	using namespace UE::AdvancedRenamer::Private;
+
+	if (FLevelEditorModule* LevelEditorModule = FModuleManager::LoadModulePtr<FLevelEditorModule>("LevelEditor"))
+	{
+		LevelEditorModule->OnLevelEditorCreated().Remove(LevelEditorCreatedDelegateHandle);
+		LevelEditorCreatedDelegateHandle.Reset();
+	}
+	ShutdownMenu();
 }
 
 void FAdvancedRenamerLevelEditorIntegration::InitializeMenu()
@@ -121,7 +142,6 @@ void FAdvancedRenamerLevelEditorIntegration::InitializeMenu()
 	const FAdvancedRenamerCommands& AdvRenCommands = FAdvancedRenamerCommands::Get();
 
 	const TAttribute<FText> TextAttribute;
-	const FSlateIcon RenameIcon = FSlateIcon(FAppStyle::GetAppStyleSetName(), "GenericCommands.Rename");
 
 	for (const FName& Menu : Menus)
 	{
@@ -129,32 +149,19 @@ void FAdvancedRenamerLevelEditorIntegration::InitializeMenu()
 		FToolMenuSection& Section = ToolMenu->FindOrAddSection(NAME_None);
 
 		Section.AddMenuEntry(
-			AdvRenCommands.RenameSelectedActors,
-			TextAttribute,
-			TextAttribute,
-			RenameIcon
+			AdvRenCommands.BatchRenameObject,
+			LOCTEXT("BatchRename", "Rename Selected Actors"),
+			LOCTEXT("BatchRenameToolTip", "Opens the Batch Renamer Panel to rename all selected actors.")
+			
 		);
 
 		Section.AddMenuEntry(
-			AdvRenCommands.RenameSharedClassActors,
-			TextAttribute,
-			TextAttribute,
-			RenameIcon
+			AdvRenCommands.BatchRenameSharedClassActors,
+			LOCTEXT("BatchRenameByClass", "Rename Actors of Selected Actor Classes")
 		);
 	}
 }
 
-void FAdvancedRenamerLevelEditorIntegration::Shutdown()
-{
-	using namespace UE::AdvancedRenamer::Private;
-
-	if (FLevelEditorModule* LevelEditorModule = FModuleManager::LoadModulePtr<FLevelEditorModule>("LevelEditor"))
-	{
-		LevelEditorModule->OnLevelEditorCreated().Remove(LevelEditorCreatedDelegateHandle);
-		LevelEditorCreatedDelegateHandle.Reset();
-	}
-}
-
 void FAdvancedRenamerLevelEditorIntegration::ShutdownMenu()
 {
 	if (UToolMenus* ToolMenus = UToolMenus::Get())
@@ -166,3 +173,5 @@ void FAdvancedRenamerLevelEditorIntegration::ShutdownMenu()
 		}
 	}
 }
+
+#undef LOCTEXT_NAMESPACE
diff --git a/Engine/Plugins/Experimental/AdvancedRenamer/Source/AdvancedRenamer/Private/Integrations/AdvancedRenamerLevelEditorIntegration.h b/Engine/Plugins/Experimental/AdvancedRenamer/Source/AdvancedRenamer/Private/Integrations/AdvancedRenamerLevelEditorIntegration.h
--- a/Engine/Plugins/Experimental/AdvancedRenamer/Source/AdvancedRenamer/Private/Integrations/AdvancedRenamerLevelEditorIntegration.h
+++ b/Engine/Plugins/Experimental/AdvancedRenamer/Source/AdvancedRenamer/Private/Integrations/AdvancedRenamerLevelEditorIntegration.h
@@ -6,7 +6,9 @@ class FAdvancedRenamerLevelEditorIntegration
 {
 public:
 	static void Initialize();
+	static void Shutdown();
+
+private:
 	static void InitializeMenu();
-	static void Shutdown();	
 	static void ShutdownMenu();
 };
