[BatchRenamer] Support auto-lettering and Formatting
Added Numbering and Lettering Section with Formatting options

diff --git a/Engine/Plugins/Experimental/AdvancedRenamer/Source/AdvancedRenamer/Private/AdvancedRenamerModule.cpp b/Engine/Plugins/Experimental/AdvancedRenamer/Source/AdvancedRenamer/Private/AdvancedRenamerModule.cpp
--- a/Engine/Plugins/Experimental/AdvancedRenamer/Source/AdvancedRenamer/Private/AdvancedRenamerModule.cpp
+++ b/Engine/Plugins/Experimental/AdvancedRenamer/Source/AdvancedRenamer/Private/AdvancedRenamerModule.cpp
@@ -6,6 +6,7 @@
 #include "AdvancedRenamerCommands.h"
 #include "AdvancedRenamerSections/AdvancedRenamerAddPrefixSuffixSection.h"
 #include "AdvancedRenamerSections/AdvancedRenamerChangeCaseSection.h"
+#include "AdvancedRenamerSections/AdvancedRenamerNumberingSection.h"
 #include "AdvancedRenamerSections/AdvancedRenamerRemovePrefixSection.h"
 #include "AdvancedRenamerSections/AdvancedRenamerRemoveSuffixSection.h"
 #include "AdvancedRenamerSections/AdvancedRenamerSearchAndReplaceSection.h"
@@ -29,12 +30,13 @@ namespace UE::AdvancedRenamer::Private
 	{
 		return SNew(SWindow)
 			.Title(LOCTEXT("AdvancedRenameWindow", "Batch Renamer"))
-			.ClientSize(FVector2D(731, 547))
+			.ClientSize(FVector2D(731, 603))
 			.SizingRule(ESizingRule::UserSized)
 			.SupportsMaximize(false)
 			.SupportsMinimize(false)
 			.MinWidth(737)
-			.MinHeight(586.5f);
+			.MaxHeight(626.f)
+			.MinHeight(626.f);
 	}
 }
 
@@ -216,6 +218,7 @@ void FAdvancedRenamerModule::RegisterDefaultSections()
 	Sections.Add(IAdvancedRenamerSection::MakeInstance<FAdvancedRenamerRemovePrefixSection>());
 	Sections.Add(IAdvancedRenamerSection::MakeInstance<FAdvancedRenamerRemoveSuffixSection>());
 	Sections.Add(IAdvancedRenamerSection::MakeInstance<FAdvancedRenamerAddPrefixSuffixSection>());
+	Sections.Add(IAdvancedRenamerSection::MakeInstance<FAdvancedRenamerNumberingSection>());
 	Sections.Add(IAdvancedRenamerSection::MakeInstance<FAdvancedRenamerChangeCaseSection>());
 }
 
diff --git a/Engine/Plugins/Experimental/AdvancedRenamer/Source/AdvancedRenamer/Private/AdvancedRenamerSections/AdvancedRenamerAddPrefixSuffixSection.cpp b/Engine/Plugins/Experimental/AdvancedRenamer/Source/AdvancedRenamer/Private/AdvancedRenamerSections/AdvancedRenamerAddPrefixSuffixSection.cpp
--- a/Engine/Plugins/Experimental/AdvancedRenamer/Source/AdvancedRenamer/Private/AdvancedRenamerSections/AdvancedRenamerAddPrefixSuffixSection.cpp
+++ b/Engine/Plugins/Experimental/AdvancedRenamer/Source/AdvancedRenamer/Private/AdvancedRenamerSections/AdvancedRenamerAddPrefixSuffixSection.cpp
@@ -24,7 +24,6 @@ void FAdvancedRenamerAddPrefixSuffixSection::Init(TSharedRef<IAdvancedRenamer> I
 {
 	FAdvancedRenamerSectionBase::Init(InRenamer);
 	Section.SectionName = TEXT("AddPrefixSuffixNumber");
-	Section.OnBeforeOperationExecutionStart().BindSP(this, &FAdvancedRenamerAddPrefixSuffixSection::ResetCurrentSuffixNumber);
 	Section.OnOperationExecuted().BindSP(this, &FAdvancedRenamerAddPrefixSuffixSection::ApplyAddPrefixSuffixNumberOperation);
 	InRenamer->AddSection(Section);
 }
@@ -54,14 +53,6 @@ TSharedRef<SWidget> FAdvancedRenamerAddPrefixSuffixSection::GetWidget()
 			[
 				CreateAddSuffix()
 			]
-			
-			// Add Number
-			+ SVerticalBox::Slot()
-			.Padding(SectionContentMiddleEntriesPadding)
-			.AutoHeight()
-			[
-				CreateAddNumber()
-			]
 		];
 }
 
@@ -69,10 +60,6 @@ void FAdvancedRenamerAddPrefixSuffixSection::ResetToDefault()
 {
 	PrefixText = FText::GetEmpty();
 	SuffixText = FText::GetEmpty();
-	bAddSuffixNumbers = false;
-	SuffixNumberStartValue = 1;
-	CurrentSuffixNumber = 1;
-	SuffixNumberStepValue = 1;
 }
 
 TSharedRef<SWidget> FAdvancedRenamerAddPrefixSuffixSection::CreateAddPrefix()
@@ -84,7 +71,7 @@ TSharedRef<SWidget> FAdvancedRenamerAddPrefixSuffixSection::CreateAddPrefix()
 		.HAlign(HAlign_Left)
 		.VAlign(VAlign_Center)
 		.AutoWidth()
-		.Padding(NameWidgetPadding)
+		.Padding(FirstWidgetPadding)
 		[
 			SNew(SBox)
 			.WidthOverride(70.f)
@@ -98,10 +85,9 @@ TSharedRef<SWidget> FAdvancedRenamerAddPrefixSuffixSection::CreateAddPrefix()
 		+ SHorizontalBox::Slot()
 		.VAlign(VAlign_Center)
 		.FillWidth(1.f)
-		.Padding(ValueWidgetPadding)
+		.Padding(LastWidgetPadding)
 		[
 			SAssignNew(PrefixTextBox, SEditableTextBox)
-			.BackgroundColor(FStyleColors::Background.GetSpecifiedColor())
 			.Font(FAdvancedRenamerStyle::Get().GetFontStyle("AdvancedRenamer.Style.RegularFont"))
 			.HintText(LOCTEXT("AR_PrefixHint", "New Prefix"))
 			.Text(this, &FAdvancedRenamerAddPrefixSuffixSection::GetPrefixText)
@@ -118,7 +104,7 @@ TSharedRef<SWidget> FAdvancedRenamerAddPrefixSuffixSection::CreateAddSuffix()
 		.HAlign(HAlign_Left)
 		.VAlign(VAlign_Center)
 		.AutoWidth()
-		.Padding(NameWidgetPadding)
+		.Padding(FirstWidgetPadding)
 		[
 			SNew(SBox)
 			.WidthOverride(70.f)
@@ -131,10 +117,9 @@ TSharedRef<SWidget> FAdvancedRenamerAddPrefixSuffixSection::CreateAddSuffix()
 		+ SHorizontalBox::Slot()
 		.VAlign(VAlign_Center)
 		.FillWidth(1.f)
-		.Padding(ValueWidgetPadding)
+		.Padding(LastWidgetPadding)
 		[
 			SAssignNew(SuffixTextBox, SEditableTextBox)
-			.BackgroundColor(FStyleColors::Background.GetSpecifiedColor())
 			.Font(FAdvancedRenamerStyle::Get().GetFontStyle("AdvancedRenamer.Style.RegularFont"))
 			.HintText(LOCTEXT("AR_SuffixHint", "New Suffix"))
 			.Text(this, &FAdvancedRenamerAddPrefixSuffixSection::GetSuffixText)
@@ -142,87 +127,6 @@ TSharedRef<SWidget> FAdvancedRenamerAddPrefixSuffixSection::CreateAddSuffix()
 		];
 }
 
-TSharedRef<SWidget> FAdvancedRenamerAddPrefixSuffixSection::CreateAddNumber()
-{
-	using namespace AdvancedRenamerSlateUtils::Default;
-
-	return SNew(SHorizontalBox)
-		+ SHorizontalBox::Slot()
-		.HAlign(HAlign_Left)
-		.VAlign(VAlign_Center)
-		.AutoWidth()
-		.Padding(NameWidgetPadding)
-		[
-			SAssignNew(SuffixNumberCheckBox, SCheckBox)
-			.IsChecked(this, &FAdvancedRenamerAddPrefixSuffixSection::IsSuffixNumberChecked)
-			.OnCheckStateChanged(this, &FAdvancedRenamerAddPrefixSuffixSection::OnSuffixNumberCheckBoxChanged)
-		]
-		+ SHorizontalBox::Slot()
-		.HAlign(HAlign_Left)
-		.VAlign(VAlign_Center)
-		.AutoWidth()
-		[
-			SNew(STextBlock)
-			.Font(FAdvancedRenamerStyle::Get().GetFontStyle("AdvancedRenamer.Style.RegularFont"))
-			.Text(LOCTEXT("AR_SuffixNumber", "Add Number"))
-		]
-		+ SHorizontalBox::Slot()
-		.HAlign(HAlign_Right)
-		.VAlign(VAlign_Center)
-		[
-			SNew(SSpacer)
-		]
-		+ SHorizontalBox::Slot()
-		.HAlign(HAlign_Right)
-		.VAlign(VAlign_Center)
-		.AutoWidth()
-		[
-			SNew(STextBlock)
-			.Font(FAdvancedRenamerStyle::Get().GetFontStyle("AdvancedRenamer.Style.RegularFont"))
-			.Text(LOCTEXT("AR_Start", "Start"))
-		]
-		+ SHorizontalBox::Slot()
-		.HAlign(HAlign_Right)
-		.VAlign(VAlign_Center)
-		.AutoWidth()
-		.Padding(ValueWidgetPadding)
-		[
-			SAssignNew(SuffixNumberStartSpinBox, SSpinBox<int32>)
-			.Style(&FAppStyle::Get(), "Menu.SpinBox")
-			.Font(FAdvancedRenamerStyle::Get().GetFontStyle("AdvancedRenamer.Style.RegularFont"))
-			.MinValue(1)
-			.MaxValue(99)
-			.Value(this, &FAdvancedRenamerAddPrefixSuffixSection::GetSuffixNumberStart)
-			.IsEnabled(this, &FAdvancedRenamerAddPrefixSuffixSection::IsSuffixNumberEnabled)
-			.OnValueChanged(this, &FAdvancedRenamerAddPrefixSuffixSection::OnSuffixNumberStartChanged)
-		]
-		+ SHorizontalBox::Slot()
-		.HAlign(HAlign_Right)
-		.VAlign(VAlign_Center)
-		.AutoWidth()
-		.Padding(AddNumberStepPadding)
-		[
-			SNew(STextBlock)
-			.Font(FAdvancedRenamerStyle::Get().GetFontStyle("AdvancedRenamer.Style.RegularFont"))
-			.Text(LOCTEXT("AR_Step", "Step"))
-		]
-		+ SHorizontalBox::Slot()
-		.HAlign(HAlign_Right)
-		.VAlign(VAlign_Center)
-		.AutoWidth()
-		.Padding(ValueWidgetPadding)
-		[
-			SAssignNew(SuffixNumberStepSpinBox, SSpinBox<int32>)
-			.Style(&FAppStyle::Get(), "Menu.SpinBox")
-			.Font(FAdvancedRenamerStyle::Get().GetFontStyle("AdvancedRenamer.Style.RegularFont"))
-			.MinValue(1)
-			.MaxValue(99)
-			.Value(this, &FAdvancedRenamerAddPrefixSuffixSection::GetSuffixNumberStep)
-			.IsEnabled(this, &FAdvancedRenamerAddPrefixSuffixSection::IsSuffixNumberEnabled)
-			.OnValueChanged(this, &FAdvancedRenamerAddPrefixSuffixSection::OnSuffixNumberStepChanged)
-		];
-}
-
 FText FAdvancedRenamerAddPrefixSuffixSection::GetPrefixText() const
 {
 	return PrefixText;
@@ -233,26 +137,6 @@ FText FAdvancedRenamerAddPrefixSuffixSection::GetSuffixText() const
 	return SuffixText;
 }
 
-int32 FAdvancedRenamerAddPrefixSuffixSection::GetSuffixNumberStart() const
-{
-	return SuffixNumberStartValue;
-}
-
-int32 FAdvancedRenamerAddPrefixSuffixSection::GetSuffixNumberStep() const
-{
-	return SuffixNumberStepValue;
-}
-
-ECheckBoxState FAdvancedRenamerAddPrefixSuffixSection::IsSuffixNumberChecked() const
-{
-	return bAddSuffixNumbers ? ECheckBoxState::Checked : ECheckBoxState::Unchecked;
-}
-
-bool FAdvancedRenamerAddPrefixSuffixSection::IsSuffixNumberEnabled() const
-{
-	return bAddSuffixNumbers;
-}
-
 void FAdvancedRenamerAddPrefixSuffixSection::OnPrefixChanged(const FText& InNewText)
 {
 	PrefixText = InNewText;
@@ -265,24 +149,6 @@ void FAdvancedRenamerAddPrefixSuffixSection::OnSuffixChanged(const FText& InNewT
 	MarkRenamerDirty();
 }
 
-void FAdvancedRenamerAddPrefixSuffixSection::OnSuffixNumberCheckBoxChanged(ECheckBoxState InNewState)
-{
-	bAddSuffixNumbers = InNewState == ECheckBoxState::Checked;
-	MarkRenamerDirty();
-}
-
-void FAdvancedRenamerAddPrefixSuffixSection::OnSuffixNumberStartChanged(int32 InNewValue)
-{
-	SuffixNumberStartValue = InNewValue;
-	MarkRenamerDirty();
-}
-
-void FAdvancedRenamerAddPrefixSuffixSection::OnSuffixNumberStepChanged(int32 InNewValue)
-{
-	SuffixNumberStepValue = InNewValue;
-	MarkRenamerDirty();
-}
-
 bool FAdvancedRenamerAddPrefixSuffixSection::CanApplyAddPrefixOperation()
 {
 	return !PrefixText.IsEmpty();
@@ -293,11 +159,6 @@ bool FAdvancedRenamerAddPrefixSuffixSection::CanApplyAddSuffixOperation()
 	return !SuffixText.IsEmpty();
 }
 
-bool FAdvancedRenamerAddPrefixSuffixSection::CanApplyAddSuffixNumberOperation()
-{
-	return IsSuffixNumberEnabled();
-}
-
 void FAdvancedRenamerAddPrefixSuffixSection::ApplyAddPrefixOperation(FString& OutOriginalName)
 {
 	OutOriginalName = PrefixText.ToString() + OutOriginalName;
@@ -308,17 +169,6 @@ void FAdvancedRenamerAddPrefixSuffixSection::ApplyAddSuffixOperation(FString& Ou
 	OutOriginalName = OutOriginalName + SuffixText.ToString();
 }
 
-void FAdvancedRenamerAddPrefixSuffixSection::ApplyAddSuffixNumberOperation(FString& OutOriginalName)
-{
-	OutOriginalName = OutOriginalName + FString::FromInt(CurrentSuffixNumber);
-	CurrentSuffixNumber += SuffixNumberStepValue;
-}
-
-void FAdvancedRenamerAddPrefixSuffixSection::ResetCurrentSuffixNumber()
-{
-	CurrentSuffixNumber = SuffixNumberStartValue;
-}
-
 void FAdvancedRenamerAddPrefixSuffixSection::ApplyAddPrefixSuffixNumberOperation(FString& OutOriginalName)
 {
 	if (CanApplyAddPrefixOperation())
@@ -330,11 +180,6 @@ void FAdvancedRenamerAddPrefixSuffixSection::ApplyAddPrefixSuffixNumberOperation
 	{
 		ApplyAddSuffixOperation(OutOriginalName);
 	}
-
-	if (CanApplyAddSuffixNumberOperation())
-	{
-		ApplyAddSuffixNumberOperation(OutOriginalName);
-	}
 }
 
 #undef LOCTEXT_NAMESPACE
diff --git a/Engine/Plugins/Experimental/AdvancedRenamer/Source/AdvancedRenamer/Private/AdvancedRenamerSections/AdvancedRenamerAddPrefixSuffixSection.h b/Engine/Plugins/Experimental/AdvancedRenamer/Source/AdvancedRenamer/Private/AdvancedRenamerSections/AdvancedRenamerAddPrefixSuffixSection.h
--- a/Engine/Plugins/Experimental/AdvancedRenamer/Source/AdvancedRenamer/Private/AdvancedRenamerSections/AdvancedRenamerAddPrefixSuffixSection.h
+++ b/Engine/Plugins/Experimental/AdvancedRenamer/Source/AdvancedRenamer/Private/AdvancedRenamerSections/AdvancedRenamerAddPrefixSuffixSection.h
@@ -19,67 +19,59 @@ public:
 
 	virtual ~FAdvancedRenamerAddPrefixSuffixSection() {}
 
+	/** Init the given section */
 	virtual void Init(TSharedRef<IAdvancedRenamer> InRenamer) override;
 
+	/** Return the widget for this section*/
 	virtual TSharedRef<SWidget> GetWidget() override;
 
+	/** Reset all values of the section to the default ones */
 	virtual void ResetToDefault() override;
 
 private:
+	/** Create the AddPrefix widget */
 	TSharedRef<SWidget> CreateAddPrefix();
 
+	/** Create the AddSuffix widget */
 	TSharedRef<SWidget> CreateAddSuffix();
 
-	TSharedRef<SWidget> CreateAddNumber();
-
+	/** Get the Prefix text */
 	FText GetPrefixText() const;
 
+	/** Get the Suffix text */
 	FText GetSuffixText() const;
 
-	int32 GetSuffixNumberStart() const;
-
-	int32 GetSuffixNumberStep() const;
-
-	ECheckBoxState IsSuffixNumberChecked() const;
-
-	bool IsSuffixNumberEnabled() const;
-
+	/** Called when the Prefix text change */
 	void OnPrefixChanged(const FText& InNewText);
 
+	/** Called when the Suffix text change */
 	void OnSuffixChanged(const FText& InNewText);
 
-	void OnSuffixNumberCheckBoxChanged(ECheckBoxState InNewState);
-
-	void OnSuffixNumberStartChanged(int32 InNewValue);
-
-	void OnSuffixNumberStepChanged(int32 InNewValue);
-
+	/** Whether the AddPrefix operation can be executed */
 	bool CanApplyAddPrefixOperation();
 
+	/** Whether the AddSuffix operation can be executed */
 	bool CanApplyAddSuffixOperation();
 
-	bool CanApplyAddSuffixNumberOperation();
-
+	/** Execute logic for the AddPrefix */
 	void ApplyAddPrefixOperation(FString& OutOriginalName);
 
+	/** Execute logic for the AddSuffix */
 	void ApplyAddSuffixOperation(FString& OutOriginalName);
 
-	void ApplyAddSuffixNumberOperation(FString& OutOriginalName);
-
-	void ResetCurrentSuffixNumber();
-
+	/** Execute logic for this section */
 	void ApplyAddPrefixSuffixNumberOperation(FString& OutOriginalName);
 
 private:
+	/** AddPrefix TextBox */
 	TSharedPtr<SEditableTextBox> PrefixTextBox;
+	
+	/** AddSuffix TextBox */
 	TSharedPtr<SEditableTextBox> SuffixTextBox;
-	TSharedPtr<SCheckBox> SuffixNumberCheckBox;
-	TSharedPtr<SSpinBox<int32>> SuffixNumberStartSpinBox;
-	TSharedPtr<SSpinBox<int32>> SuffixNumberStepSpinBox;
+
+	/** AddPrefix Text */
 	FText PrefixText;
+
+	/** AddSuffix text */
 	FText SuffixText;
-	bool bAddSuffixNumbers;
-	int32 SuffixNumberStartValue;
-	int32 SuffixNumberStepValue;
-	int32 CurrentSuffixNumber;
 };
diff --git a/Engine/Plugins/Experimental/AdvancedRenamer/Source/AdvancedRenamer/Private/AdvancedRenamerSections/AdvancedRenamerNumberingSection.cpp b/Engine/Plugins/Experimental/AdvancedRenamer/Source/AdvancedRenamer/Private/AdvancedRenamerSections/AdvancedRenamerNumberingSection.cpp
new file mode 100644
--- /dev/null
+++ b/Engine/Plugins/Experimental/AdvancedRenamer/Source/AdvancedRenamer/Private/AdvancedRenamerSections/AdvancedRenamerNumberingSection.cpp
@@ -0,0 +1,513 @@
+// Copyright Epic Games, Inc. All Rights Reserved.
+
+#include "AdvancedRenamerNumberingSection.h"
+
+#include "AdvancedRenamerStyle.h"
+#include "IAdvancedRenamer.h"
+#include "Utils/AdvancedRenamerSlateUtils.h"
+#include "Widgets/Input/SComboBox.h"
+#include "Widgets/Input/SEditableTextBox.h"
+#include "Widgets/Input/SSegmentedControl.h"
+#include "Widgets/Input/SSpinBox.h"
+#include "Widgets/Layout/SBox.h"
+#include "Widgets/Layout/SWidgetSwitcher.h"
+
+#define LOCTEXT_NAMESPACE "AdvancedRenamerNumberingSection"
+
+TArray<FText> FAdvancedRenamerNumberingSection::ComboBoxTextOptions = { LOCTEXT("AR_FirstFormat", "0"), LOCTEXT("AR_SecondFormat", "00"), LOCTEXT("AR_ThirdFormat", "000"), LOCTEXT("AR_FourthFormat", "0000") };
+
+namespace AdvancedRenamerNumberingSectionUtils
+{
+	static constexpr TCHAR FirstUpperLetter = TEXT('A');
+	static constexpr TCHAR FirstLowerLetter = TEXT('a');
+	static constexpr int32 LetteringStep = 1;
+
+	/** Return the reference letter to use based on the given parameter */
+	TCHAR GetReferenceLetter(bool bInIsLower)
+	{
+		if (bInIsLower)
+		{
+			return FirstLowerLetter;
+		}
+		return FirstUpperLetter;
+	}
+}
+
+FAdvancedRenamerNumberingSection::FAdvancedRenamerNumberingSection()
+{
+	FAdvancedRenamerNumberingSection::ResetToDefault();
+}
+
+void FAdvancedRenamerNumberingSection::Init(TSharedRef<IAdvancedRenamer> InRenamer)
+{
+	FAdvancedRenamerSectionBase::Init(InRenamer);
+
+	ComboBoxSourceOptions.Empty(ComboBoxTextOptions.Num());
+	for (int32 Index = 0; Index < ComboBoxTextOptions.Num(); Index++)
+	{
+		ComboBoxSourceOptions.Add(MakeShared<int32>(Index));
+	}
+
+	Section.SectionName = TEXT("NumberingLettering");
+	Section.OnBeforeOperationExecutionStart().BindSP(this, &FAdvancedRenamerNumberingSection::ResetCurrentNumbering);
+	Section.OnOperationExecuted().BindSP(this, &FAdvancedRenamerNumberingSection::ApplyNumbering);
+	InRenamer->AddSection(Section);
+}
+
+TSharedRef<SWidget> FAdvancedRenamerNumberingSection::GetWidget()
+{
+	using namespace AdvancedRenamerSlateUtils::Default;
+
+	return SNew(SBorder)
+		.BorderImage(FAdvancedRenamerStyle::Get().GetBrush("AdvancedRenamer.Style.BackgroundBorder"))
+		.Content()
+		[
+			SNew(SVerticalBox)
+
+			// Numbering Header
+			+ SVerticalBox::Slot()
+			.Padding(SectionContentFirstEntryPadding)
+			.AutoHeight()
+			[
+				SNew(SHorizontalBox)
+
+				+ SHorizontalBox::Slot()
+				.VAlign(VAlign_Center)
+				.AutoWidth()
+				.Padding(FirstWidgetPadding)
+				[
+					SAssignNew(AddNumberingCheckBox, SCheckBox)
+					.IsChecked(this, &FAdvancedRenamerNumberingSection::IsAddNumberingChecked)
+					.OnCheckStateChanged(this, &FAdvancedRenamerNumberingSection::OnAddNumberingCheckBoxChanged)
+				]
+
+				+ SHorizontalBox::Slot()
+				.Padding(FirstWidgetPadding)
+				.VAlign(VAlign_Center)
+				.AutoWidth()
+				[
+					SNew(STextBlock)
+					.Font(FAdvancedRenamerStyle::Get().GetFontStyle("AdvancedRenamer.Style.RegularFont"))
+					.Text(LOCTEXT("AR_AddNumberingLettering", "Auto Increment"))
+				]
+
+				+ SHorizontalBox::Slot()
+				.Padding(LastWidgetPadding)
+				.VAlign(VAlign_Center)
+				.FillWidth(0.5f)
+				[
+					SNew(SBox)
+					.HeightOverride(25.f)
+					[
+						SNew(SSegmentedControl<EAdvancedRenamerNumberingType>)
+						.TextStyle(&FAppStyle::Get().GetWidgetStyle<FTextBlockStyle>("NormalText"))
+						.IsEnabled(this, &FAdvancedRenamerNumberingSection::IsAddNumberingEnabled)
+						.UniformPadding(FMargin(2.f, 0.f))
+						.SupportsEmptySelection(false)
+						.SupportsMultiSelection(false)
+						.Value(this, &FAdvancedRenamerNumberingSection::GetAddNumberingType)
+						.OnValueChanged(this, &FAdvancedRenamerNumberingSection::OnAddNumberingTypeChanged)
+									
+						+ SSegmentedControl<EAdvancedRenamerNumberingType>::Slot(EAdvancedRenamerNumberingType::Number)
+						.Text(LOCTEXT("AR_NumberingType", "#"))
+									
+						+ SSegmentedControl<EAdvancedRenamerNumberingType>::Slot(EAdvancedRenamerNumberingType::Letter)
+						.Text(LOCTEXT("AR_LetteringType", "Aa"))
+					]
+				]
+			]
+
+			// Add Number/Letter and Format
+			+ SVerticalBox::Slot()
+			.Padding(SectionContentFirstEntryPadding)
+			.AutoHeight()
+			[
+				SNew(SBox)
+				.HeightOverride(24.5f)
+				[
+					SNew(SHorizontalBox)
+
+					// Formatting
+					+ SHorizontalBox::Slot()
+					.VAlign(VAlign_Center)
+					.Padding(FirstWidgetPadding)
+					.FillWidth(0.3f)
+					[
+						SNew(SHorizontalBox)
+						.Visibility(this, &FAdvancedRenamerNumberingSection::IsAddNumberingTypeNumber)
+
+						+ SHorizontalBox::Slot()
+						.VAlign(VAlign_Center)
+						.Padding(FirstWidgetPadding)
+						.AutoWidth()
+						[
+							SNew(STextBlock)
+							.Font(FAdvancedRenamerStyle::Get().GetFontStyle("AdvancedRenamer.Style.RegularFont"))
+							.Text(LOCTEXT("AR_FormattingLabel", "Format"))
+						]
+
+						+ SHorizontalBox::Slot()
+						.VAlign(VAlign_Center)
+						.FillWidth(1.f)
+						[
+							SAssignNew(FormattingComboBox, SComboBox<TSharedPtr<int32>>)
+							.IsEnabled(this, &FAdvancedRenamerNumberingSection::IsAddNumberingEnabled)
+							.OptionsSource(&ComboBoxSourceOptions)
+							.InitiallySelectedItem(ComboBoxSourceOptions[0])
+							.ToolTipText(LOCTEXT("AR_FormatTooltip", "Select the format you prefer for the numbering"))
+							.OnGenerateWidget(this, &FAdvancedRenamerNumberingSection::OnGenerateFormatWidget)
+							.OnSelectionChanged(this, &FAdvancedRenamerNumberingSection::OnFormatSelectionChanged)
+							.Content()
+							[
+								SNew(STextBlock)
+								.Text(this, &FAdvancedRenamerNumberingSection::GetCurrentFormatText)
+							]
+						]
+					]
+
+					// Start
+					+ SHorizontalBox::Slot()
+					.VAlign(VAlign_Center)
+					.Padding(MiddleWidgetPadding)
+					.FillWidth(0.3f)
+					[
+						SNew(SHorizontalBox)
+
+						+ SHorizontalBox::Slot()
+						.VAlign(VAlign_Center)
+						.Padding(FirstWidgetPadding)
+						.AutoWidth()
+						[
+							SNew(STextBlock)
+							.Font(FAdvancedRenamerStyle::Get().GetFontStyle("AdvancedRenamer.Style.RegularFont"))
+							.Text(LOCTEXT("AR_Start", "Start"))
+						]
+
+						+ SHorizontalBox::Slot()
+						.VAlign(VAlign_Center)
+						.FillWidth(1.f)
+						[
+							SNew(SWidgetSwitcher)
+							.WidgetIndex(this, &FAdvancedRenamerNumberingSection::GetNumberingIndex)
+
+							// Add Number
+							+ SWidgetSwitcher::Slot()
+							[
+								SAssignNew(AddNumberStartSpinBox, SSpinBox<int32>)
+								.Style(&FAppStyle::Get(), "Menu.SpinBox")
+								.Font(FAdvancedRenamerStyle::Get().GetFontStyle("AdvancedRenamer.Style.RegularFont"))
+								.MinValue(0)
+								.Value(this, &FAdvancedRenamerNumberingSection::GetAddNumberValue)
+								.IsEnabled(this, &FAdvancedRenamerNumberingSection::IsAddNumberingEnabled)
+								.OnValueChanged(this, &FAdvancedRenamerNumberingSection::OnAddNumberChanged)
+							]
+
+							// Add Letter
+							+ SWidgetSwitcher::Slot()
+							[
+								SAssignNew(AddLetterStartTextBox, SEditableTextBox)
+								.Font(FAdvancedRenamerStyle::Get().GetFontStyle("AdvancedRenamer.Style.RegularFont"))
+								.OnVerifyTextChanged(this, &FAdvancedRenamerNumberingSection::OnLetteringVerifyText)
+								.Text(this, &FAdvancedRenamerNumberingSection::GetAddLetterText)
+								.IsEnabled(this, &FAdvancedRenamerNumberingSection::IsAddNumberingEnabled)
+								.OnTextChanged(this, &FAdvancedRenamerNumberingSection::OnAddLetterChanged)
+							]
+						]
+					]
+
+					// Step
+					+ SHorizontalBox::Slot()
+					.VAlign(VAlign_Center)
+					.Padding(LastWidgetPadding)
+					.FillWidth(0.3f)
+					[
+						SNew(SHorizontalBox)
+						.Visibility(this, &FAdvancedRenamerNumberingSection::IsAddNumberingTypeNumber)
+
+						+ SHorizontalBox::Slot()
+						.Padding(FirstWidgetPadding)
+						.VAlign(VAlign_Center)
+						.AutoWidth()
+						[
+							SNew(STextBlock)
+							.Font(FAdvancedRenamerStyle::Get().GetFontStyle("AdvancedRenamer.Style.RegularFont"))
+							.Text(LOCTEXT("AR_Step", "Step"))
+						]
+
+						+ SHorizontalBox::Slot()
+						.VAlign(VAlign_Center)
+						.FillWidth(1.f)
+						[
+							SAssignNew(AddNumberStepSpinBox, SSpinBox<int32>)
+							.Style(&FAppStyle::Get(), "Menu.SpinBox")
+							.Font(FAdvancedRenamerStyle::Get().GetFontStyle("AdvancedRenamer.Style.RegularFont"))
+							.MinValue(1)
+							.MaxValue(99)
+							.Value(this, &FAdvancedRenamerNumberingSection::GetAddNumberingStepValue)
+							.IsEnabled(this, &FAdvancedRenamerNumberingSection::IsAddNumberingEnabled)
+							.OnValueChanged(this, &FAdvancedRenamerNumberingSection::OnAddNumberingStepChanged)
+						]
+					]
+				]
+			]
+		];
+}
+
+void FAdvancedRenamerNumberingSection::ResetToDefault()
+{
+	bIsAddNumberingSectionEnabled = false;
+	bIsAddLetteringInputCorrect = true;
+	AddNumberingType = EAdvancedRenamerNumberingType::Number;
+	AddNumberingWidgetSwitcherIndex = 0;
+	AddNumberValue = 0;
+	AddLetterText = LOCTEXT("AR_AddLetter", "A");
+	AddNumberingStepValue = 1;
+	CurrentAddNumberValue = 1;
+	CurrentAddLetterString = TEXT("");
+	CurrentFormatChosen = 0;
+	if (FormattingComboBox.IsValid() && !ComboBoxSourceOptions.IsEmpty())
+	{
+		FormattingComboBox->SetSelectedItem(ComboBoxSourceOptions[0]);
+	}
+}
+
+ECheckBoxState FAdvancedRenamerNumberingSection::IsAddNumberingChecked() const
+{
+	return bIsAddNumberingSectionEnabled ? ECheckBoxState::Checked : ECheckBoxState::Unchecked;
+}
+
+bool FAdvancedRenamerNumberingSection::IsAddNumberingEnabled() const
+{
+	return bIsAddNumberingSectionEnabled;
+}
+
+EVisibility FAdvancedRenamerNumberingSection::IsAddNumberingTypeNumber() const
+{
+	return AddNumberingType == EAdvancedRenamerNumberingType::Number ? EVisibility::Visible : EVisibility::Collapsed;
+}
+
+EAdvancedRenamerNumberingType FAdvancedRenamerNumberingSection::GetAddNumberingType() const
+{
+	return AddNumberingType;
+}
+
+FText FAdvancedRenamerNumberingSection::GetCurrentFormatText() const
+{
+	if (ComboBoxTextOptions.IsValidIndex(CurrentFormatChosen))
+	{
+		return ComboBoxTextOptions[CurrentFormatChosen];
+	}
+	return FText::GetEmpty();
+}
+
+int32 FAdvancedRenamerNumberingSection::GetNumberingIndex() const
+{
+	return AddNumberingWidgetSwitcherIndex;
+}
+
+int32 FAdvancedRenamerNumberingSection::GetAddNumberValue() const
+{
+	return AddNumberValue;
+}
+
+FText FAdvancedRenamerNumberingSection::GetAddLetterText() const
+{
+	return AddLetterText;
+}
+
+int32 FAdvancedRenamerNumberingSection::GetAddNumberingStepValue() const
+{
+	return AddNumberingStepValue;
+}
+
+FText FAdvancedRenamerNumberingSection::GetFormatTextForIndex(TSharedPtr<int32> InFormatIndex) const
+{
+	if (InFormatIndex.IsValid())
+	{
+		int32 Index = *InFormatIndex.Get();
+		if (ComboBoxTextOptions.IsValidIndex(Index))
+		{
+			return ComboBoxTextOptions[Index];
+		}
+	}
+	return FText::GetEmpty();
+}
+
+FString FAdvancedRenamerNumberingSection::GetFormattedNumber() const
+{
+	if (CurrentFormatChosen == 0)
+	{
+		return FString::Printf(TEXT("%01d"), CurrentAddNumberValue);
+	}
+	if (CurrentFormatChosen == 1)
+	{
+		return FString::Printf(TEXT("%02d"), CurrentAddNumberValue);
+	}
+	if (CurrentFormatChosen == 2)
+	{
+		return FString::Printf(TEXT("%03d"), CurrentAddNumberValue);
+	}
+	if (CurrentFormatChosen == 3)
+	{
+		return FString::Printf(TEXT("%04d"), CurrentAddNumberValue);
+	}
+
+	return TEXT("Error: Format Not Valid");
+}
+
+TSharedRef<SWidget> FAdvancedRenamerNumberingSection::OnGenerateFormatWidget(TSharedPtr<int32> InOption) const
+{
+	using namespace AdvancedRenamerSlateUtils::Default;
+
+	return SNew(SBox)
+		.Padding(VerticalPadding)
+		[
+			SNew(STextBlock)
+			.Text(this, &FAdvancedRenamerNumberingSection::GetFormatTextForIndex, InOption)
+		];
+}
+
+void FAdvancedRenamerNumberingSection::OnFormatSelectionChanged(TSharedPtr<int32> InNewFormat, ESelectInfo::Type InSelectInfo)
+{
+	if (InNewFormat.IsValid())
+	{
+		CurrentFormatChosen = *InNewFormat.Get();
+	}
+	else
+	{
+		// Set the format to -1 so that the widget will display an error text showing that an error occured
+		CurrentFormatChosen = -1;
+	}
+	MarkRenamerDirty();
+}
+
+void FAdvancedRenamerNumberingSection::OnAddNumberingCheckBoxChanged(ECheckBoxState InNewState)
+{
+	bIsAddNumberingSectionEnabled = InNewState == ECheckBoxState::Checked;
+	MarkRenamerDirty();
+}
+
+void FAdvancedRenamerNumberingSection::OnAddNumberingTypeChanged(EAdvancedRenamerNumberingType InNewValue)
+{
+	AddNumberingType = InNewValue;
+	switch (AddNumberingType)
+	{
+		case EAdvancedRenamerNumberingType::Number:
+			AddNumberingWidgetSwitcherIndex = 0;
+			break;
+
+		case EAdvancedRenamerNumberingType::Letter:
+			AddNumberingWidgetSwitcherIndex = 1;
+			break;
+	}
+	MarkRenamerDirty();
+}
+
+void FAdvancedRenamerNumberingSection::OnAddNumberChanged(int32 InNewValue)
+{
+	AddNumberValue = InNewValue;
+	MarkRenamerDirty();
+}
+
+void FAdvancedRenamerNumberingSection::OnAddLetterChanged(const FText& InNewText)
+{
+	AddLetterText = InNewText;
+	MarkRenamerDirty();
+}
+
+bool FAdvancedRenamerNumberingSection::OnLetteringVerifyText(const FText& InNewText, FText& OutErrorText)
+{
+	FString NewString = InNewText.ToString();
+
+	for (int32 StringIndex = 0; StringIndex < NewString.Len(); StringIndex++)
+	{
+		if (!FChar::IsAlpha(NewString[StringIndex]))
+		{
+			OutErrorText = LOCTEXT("AR_LetterInputNotValid", "Only alphabetic character are allowed as input");
+			bIsAddLetteringInputCorrect = false;
+			return false;
+		}
+	}
+
+	bIsAddLetteringInputCorrect = true;
+	return true;
+}
+
+void FAdvancedRenamerNumberingSection::OnAddNumberingStepChanged(int32 InNewValue)
+{
+	AddNumberingStepValue = InNewValue;
+	MarkRenamerDirty();
+}
+
+bool FAdvancedRenamerNumberingSection::CanApplyAddNumberOperation() const
+{
+	return AddNumberingType == EAdvancedRenamerNumberingType::Number;
+}
+
+bool FAdvancedRenamerNumberingSection::CanApplyAddLetterOperation() const
+{
+	return AddNumberingType == EAdvancedRenamerNumberingType::Letter && bIsAddLetteringInputCorrect;
+}
+
+void FAdvancedRenamerNumberingSection::ApplyAddNumberOperation(FString& OutOriginalName)
+{
+	OutOriginalName += GetFormattedNumber();
+	CurrentAddNumberValue += AddNumberingStepValue;
+}
+
+void FAdvancedRenamerNumberingSection::ApplyAddLetterOperation(FString& OutOriginalName)
+{
+	using namespace AdvancedRenamerNumberingSectionUtils;
+
+	if (CurrentAddLetterString.IsEmpty())
+	{
+		return;
+	}
+
+	OutOriginalName += CurrentAddLetterString;
+
+	int32 CountToUse = LetteringStep;
+
+	for (int32 CurrentLetteringChar = (CurrentAddLetterString.Len() - 1); CurrentLetteringChar >= 0; --CurrentLetteringChar)
+	{
+		if (CountToUse == 0)
+		{
+			break;
+		}
+
+		TCHAR& CurrentChar = CurrentAddLetterString[CurrentLetteringChar];
+		TCHAR LetterToUseAsReference = GetReferenceLetter(FChar::IsLower(CurrentChar));
+
+		int32 Difference = (CurrentChar - LetterToUseAsReference + CountToUse);
+		CurrentChar = LetterToUseAsReference + (Difference % 26);
+		CountToUse = (float)Difference / 26.f;
+
+		if (!CurrentAddLetterString.IsValidIndex(CurrentLetteringChar - 1) && CountToUse > 0)
+		{
+			CurrentAddLetterString.InsertAt(0, LetterToUseAsReference);
+		}
+	}
+}
+
+void FAdvancedRenamerNumberingSection::ResetCurrentNumbering()
+{
+	CurrentAddNumberValue = AddNumberValue;
+	CurrentAddLetterString = AddLetterText.ToString();
+}
+
+void FAdvancedRenamerNumberingSection::ApplyNumbering(FString& OutOriginalName)
+{
+	if (bIsAddNumberingSectionEnabled)
+	{
+		if (CanApplyAddNumberOperation())
+		{
+			ApplyAddNumberOperation(OutOriginalName);
+		}
+		else if (CanApplyAddLetterOperation())
+		{
+			ApplyAddLetterOperation(OutOriginalName);
+		}
+	}
+}
+
+#undef LOCTEXT_NAMESPACE
diff --git a/Engine/Plugins/Experimental/AdvancedRenamer/Source/AdvancedRenamer/Private/AdvancedRenamerSections/AdvancedRenamerNumberingSection.h b/Engine/Plugins/Experimental/AdvancedRenamer/Source/AdvancedRenamer/Private/AdvancedRenamerSections/AdvancedRenamerNumberingSection.h
new file mode 100644
--- /dev/null
+++ b/Engine/Plugins/Experimental/AdvancedRenamer/Source/AdvancedRenamer/Private/AdvancedRenamerSections/AdvancedRenamerNumberingSection.h
@@ -0,0 +1,166 @@
+// Copyright Epic Games, Inc. All Rights Reserved.
+
+#pragma once
+
+#include "AdvancedRenamerSectionBase.h"
+#include "Layout/Visibility.h"
+#include "Styling/SlateTypes.h"
+
+class SCheckBox;
+template <typename InOptionType>
+class SComboBox;
+class SEditableTextBox;
+template <typename InNumericType>
+class SSpinBox;
+class SWidget;
+
+enum class EAdvancedRenamerNumberingType
+{
+	Number,
+	Letter
+};
+
+class FAdvancedRenamerNumberingSection : public FAdvancedRenamerSectionBase
+{
+public:
+	FAdvancedRenamerNumberingSection();
+
+	virtual ~FAdvancedRenamerNumberingSection() {}
+
+	/** Init the given section */
+	virtual void Init(TSharedRef<IAdvancedRenamer> InRenamer) override;
+
+	/** Return the widget for this section*/
+	virtual TSharedRef<SWidget> GetWidget() override;
+
+	/** Reset all values of the section to the default ones */
+	virtual void ResetToDefault() override;
+
+private:
+	/** If checked the section will be enabled */
+	ECheckBoxState IsAddNumberingChecked() const;
+
+	/** Whether the section is enabled */
+	bool IsAddNumberingEnabled() const;
+
+	/** Visible if the type is Number, Collapsed otherwise */
+	EVisibility IsAddNumberingTypeNumber() const;
+
+	/** Get the type of the AddNumbering, either by Number or by Letter */
+	EAdvancedRenamerNumberingType GetAddNumberingType() const;
+
+	/** Get the current Text Format */
+	FText GetCurrentFormatText() const;
+
+	/** Get the Index of the WidgetSwitcher to either show the Number or the Letter */
+	int32 GetNumberingIndex() const;
+
+	/** Get the starting Number of the AddNumber SpinBox */
+	int32 GetAddNumberValue() const;
+
+	/** Get the starting Letter of the AddLetter TextBox */
+	FText GetAddLetterText() const;
+
+	/** Get the starting Number of the AddNumbering Step SpinBox */
+	int32 GetAddNumberingStepValue() const;
+
+	/** Get the Text Format for the given InFormatIndex */
+	FText GetFormatTextForIndex(TSharedPtr<int32> InFormatIndex) const;
+
+	/** Get the formatted Number Text */
+	FString GetFormattedNumber() const;
+
+	/** Generate the Format option widget for the ComboBox */
+	TSharedRef<SWidget> OnGenerateFormatWidget(TSharedPtr<int32> InOption) const;
+
+	/** Called when the Format selected change */
+	void OnFormatSelectionChanged(TSharedPtr<int32> InNewFormat, ESelectInfo::Type InSelectInfo);
+
+	/** Called when the AddNumbering CheckBox state change */
+	void OnAddNumberingCheckBoxChanged(ECheckBoxState InNewState);
+
+	/** Called when the AddNumbering type changed, it will be either by Number or Letter */
+	void OnAddNumberingTypeChanged(EAdvancedRenamerNumberingType InNewValue);
+
+	/** Called when the AddNumber number change */
+	void OnAddNumberChanged(int32 InNewValue);
+
+	/** Called when the AddLetter letter change */
+	void OnAddLetterChanged(const FText& InNewText);
+
+	/** Called to verify that the Letter input is correct and allowed */
+	bool OnLetteringVerifyText(const FText& InNewText, FText& OutErrorText);
+
+	/** Called when the AddNumbering Step number change */
+	void OnAddNumberingStepChanged(int32 InNewValue);
+
+	/** Whether the AddNumber operation can be executed */
+	bool CanApplyAddNumberOperation() const;
+
+	/** Whether the AddLetter operation can be executed */
+	bool CanApplyAddLetterOperation() const;
+
+	/** Execute logic for the AddNumber */
+	void ApplyAddNumberOperation(FString& OutOriginalName);
+
+	/** Execute logic for the AddLetter */
+	void ApplyAddLetterOperation(FString& OutOriginalName);
+
+	/** Reset current numbering */
+	void ResetCurrentNumbering();
+
+	/** Execute logic for this section */
+	void ApplyNumbering(FString& OutOriginalName);
+
+private:
+	/** Section enabler CheckBox */
+	TSharedPtr<SCheckBox> AddNumberingCheckBox;
+
+	/** Starting Number SpinBox for the AddNumber */
+	TSharedPtr<SSpinBox<int32>> AddNumberStartSpinBox;
+
+	/** Starting Letter TextBox for the AddLetter */
+	TSharedPtr<SEditableTextBox> AddLetterStartTextBox;
+
+	/** Step Number SpinBox for the AddNumbering section */
+	TSharedPtr<SSpinBox<int32>> AddNumberStepSpinBox;
+
+	/** Step Number SpinBox for the AddNumbering section */
+	TSharedPtr<SComboBox<TSharedPtr<int32>>> FormattingComboBox;
+
+	/** AddNumbering section enabler, section is enabled if true */
+	bool bIsAddNumberingSectionEnabled;
+
+	/** AddNumbering section enabler, section is enabled if true */
+	bool bIsAddLetteringInputCorrect;
+
+	/** AddNumbering type, either Number or Letter */
+	EAdvancedRenamerNumberingType AddNumberingType;
+
+	/** Index of the WidgetSwitcher between Number and Letter */
+	int32 AddNumberingWidgetSwitcherIndex;
+
+	/** AddNumber start/stop value */
+	int32 AddNumberValue;
+
+	/** AddLetter start/stop value */
+	FText AddLetterText;
+
+	/** AddNumber step value */
+	int32 AddNumberingStepValue;
+
+	/** Current AddNumber value */
+	int32 CurrentAddNumberValue;
+
+	/** Current AddLetter value */
+	FString CurrentAddLetterString;
+
+	/** Current format chosen */
+	int32 CurrentFormatChosen;
+
+	/** Source option for the DropDown Menu of the ComboBox */
+	TArray<TSharedPtr<int32>> ComboBoxSourceOptions;
+
+	/** Text option based on the SourceOption */
+	static TArray<FText> ComboBoxTextOptions;
+};
diff --git a/Engine/Plugins/Experimental/AdvancedRenamer/Source/AdvancedRenamer/Private/AdvancedRenamerSections/AdvancedRenamerRemovePrefixSection.cpp b/Engine/Plugins/Experimental/AdvancedRenamer/Source/AdvancedRenamer/Private/AdvancedRenamerSections/AdvancedRenamerRemovePrefixSection.cpp
--- a/Engine/Plugins/Experimental/AdvancedRenamer/Source/AdvancedRenamer/Private/AdvancedRenamerSections/AdvancedRenamerRemovePrefixSection.cpp
+++ b/Engine/Plugins/Experimental/AdvancedRenamer/Source/AdvancedRenamer/Private/AdvancedRenamerSections/AdvancedRenamerRemovePrefixSection.cpp
@@ -48,17 +48,15 @@ TSharedRef<SWidget> FAdvancedRenamerRemovePrefixSection::GetWidget()
 				.HAlign(HAlign_Left)
 				.VAlign(VAlign_Center)
 				.AutoWidth()
-				.Padding(NameWidgetPadding)
+				.Padding(FirstWidgetPadding)
 				[
 					SAssignNew(RemoveOldPrefixCheckBox, SCheckBox)
 					.IsChecked(this, &FAdvancedRenamerRemovePrefixSection::IsRemoveOldPrefixChecked)
 					.OnCheckStateChanged(this, &FAdvancedRenamerRemovePrefixSection::OnRemoveOldPrefixCheckBoxChanged)
 				]
 				+ SHorizontalBox::Slot()
-				.HAlign(HAlign_Left)
 				.VAlign(VAlign_Center)
 				.AutoWidth()
-				.Padding(VerticalAddNumberPadding)
 				[
 					SNew(STextBlock)
 					.Font(FAdvancedRenamerStyle::Get().GetFontStyle("AdvancedRenamer.Style.RegularFont"))
@@ -77,7 +75,7 @@ TSharedRef<SWidget> FAdvancedRenamerRemovePrefixSection::GetWidget()
 				.HAlign(HAlign_Left)
 				.VAlign(VAlign_Center)
 				.AutoWidth()
-				.Padding(NameWidgetPadding)
+				.Padding(FirstWidgetPadding)
 				[
 					SNew(SBox)
 					.HeightOverride(25.f)
@@ -100,7 +98,7 @@ TSharedRef<SWidget> FAdvancedRenamerRemovePrefixSection::GetWidget()
 					]
 				]
 				+ SHorizontalBox::Slot()
-				.Padding(ValueWidgetPadding)
+				.Padding(LastWidgetPadding)
 				[
 					SNew(SWidgetSwitcher)
 					.IsEnabled(this, &FAdvancedRenamerRemovePrefixSection::IsRemoveOldPrefixEnabled)
@@ -109,7 +107,6 @@ TSharedRef<SWidget> FAdvancedRenamerRemovePrefixSection::GetWidget()
 					+ SWidgetSwitcher::Slot()
 					[
 						SAssignNew(PrefixSeparatorTextBox, SEditableTextBox)
-						.BackgroundColor(FStyleColors::Background.GetSpecifiedColor())
 						.Font(FAdvancedRenamerStyle::Get().GetFontStyle("AdvancedRenamer.Style.RegularFont"))
 						.Text(this, &FAdvancedRenamerRemovePrefixSection::GetPrefixSeparatorText)
 						.OnVerifyTextChanged(this, &FAdvancedRenamerRemovePrefixSection::OnPrefixSeparatorVerifyTextChanged)
diff --git a/Engine/Plugins/Experimental/AdvancedRenamer/Source/AdvancedRenamer/Private/AdvancedRenamerSections/AdvancedRenamerRemoveSuffixSection.cpp b/Engine/Plugins/Experimental/AdvancedRenamer/Source/AdvancedRenamer/Private/AdvancedRenamerSections/AdvancedRenamerRemoveSuffixSection.cpp
--- a/Engine/Plugins/Experimental/AdvancedRenamer/Source/AdvancedRenamer/Private/AdvancedRenamerSections/AdvancedRenamerRemoveSuffixSection.cpp
+++ b/Engine/Plugins/Experimental/AdvancedRenamer/Source/AdvancedRenamer/Private/AdvancedRenamerSections/AdvancedRenamerRemoveSuffixSection.cpp
@@ -48,17 +48,15 @@ TSharedRef<SWidget> FAdvancedRenamerRemoveSuffixSection::GetWidget()
 				.HAlign(HAlign_Left)
 				.VAlign(VAlign_Center)
 				.AutoWidth()
-				.Padding(NameWidgetPadding)
+				.Padding(FirstWidgetPadding)
 				[
 					SAssignNew(RemoveOldSuffixCheckBox, SCheckBox)
 					.IsChecked(this, &FAdvancedRenamerRemoveSuffixSection::IsRemoveOldSuffixChecked)
 					.OnCheckStateChanged(this, &FAdvancedRenamerRemoveSuffixSection::OnRemoveOldSuffixCheckBoxChanged)
 				]
 				+ SHorizontalBox::Slot()
-				.HAlign(HAlign_Left)
 				.VAlign(VAlign_Center)
 				.AutoWidth()
-				.Padding(VerticalAddNumberPadding)
 				[
 					SNew(STextBlock)
 					.Font(FAdvancedRenamerStyle::Get().GetFontStyle("AdvancedRenamer.Style.RegularFont"))
@@ -77,7 +75,7 @@ TSharedRef<SWidget> FAdvancedRenamerRemoveSuffixSection::GetWidget()
 				.HAlign(HAlign_Left)
 				.VAlign(VAlign_Center)
 				.AutoWidth()
-				.Padding(NameWidgetPadding)
+				.Padding(FirstWidgetPadding)
 				[
 					SNew(SBox)
 					.HeightOverride(25.f)
@@ -100,7 +98,7 @@ TSharedRef<SWidget> FAdvancedRenamerRemoveSuffixSection::GetWidget()
 					]
 				]
 				+ SHorizontalBox::Slot()
-				.Padding(ValueWidgetPadding)
+				.Padding(LastWidgetPadding)
 				[
 					SNew(SWidgetSwitcher)
 					.IsEnabled(this, &FAdvancedRenamerRemoveSuffixSection::IsRemoveOldSuffixEnabled)
@@ -109,7 +107,6 @@ TSharedRef<SWidget> FAdvancedRenamerRemoveSuffixSection::GetWidget()
 					+ SWidgetSwitcher::Slot()
 					[
 						SAssignNew(SuffixSeparatorTextBox, SEditableTextBox)
-						.BackgroundColor(FStyleColors::Background.GetSpecifiedColor())
 						.Font(FAdvancedRenamerStyle::Get().GetFontStyle("AdvancedRenamer.Style.RegularFont"))
 						.Text(this, &FAdvancedRenamerRemoveSuffixSection::GetSuffixSeparatorText)
 						.OnVerifyTextChanged(this, &FAdvancedRenamerRemoveSuffixSection::OnSuffixSeparatorVerifyTextChanged)
@@ -139,23 +136,20 @@ TSharedRef<SWidget> FAdvancedRenamerRemoveSuffixSection::GetWidget()
 				+ SHorizontalBox::Slot()
 				.HAlign(HAlign_Left)
 				.VAlign(VAlign_Center)
-				.Padding(NameWidgetPadding)
+				.Padding(FirstWidgetPadding)
 				.AutoWidth()
 				[
 					SAssignNew(SuffixRemoveNumberCheckBox, SCheckBox)
-					.IsEnabled(this, &FAdvancedRenamerRemoveSuffixSection::IsRemoveOldSuffixEnabled)
 					.IsChecked(this, &FAdvancedRenamerRemoveSuffixSection::IsSuffixRemoveNumberChecked)
 					.OnCheckStateChanged(this, &FAdvancedRenamerRemoveSuffixSection::OnSuffixRemoveNumberCheckBoxChanged)
 				]
 				+ SHorizontalBox::Slot()
-				.HAlign(HAlign_Left)
 				.VAlign(VAlign_Center)
-				.Padding(VerticalAddNumberPadding)
 				.AutoWidth()
 				[
 					SNew(STextBlock)
 					.Font(FAdvancedRenamerStyle::Get().GetFontStyle("AdvancedRenamer.Style.RegularFont"))
-					.Text(LOCTEXT("AR_SuffixRemoveNumber", "Remove Number"))
+					.Text(LOCTEXT("AR_SuffixRemoveNumber", "Remove Old Numbering"))
 				]
 			]
 		];
@@ -306,13 +300,13 @@ void FAdvancedRenamerRemoveSuffixSection::ApplyRemoveSuffixNumbers(FString& OutO
 
 void FAdvancedRenamerRemoveSuffixSection::ApplyRemoveSuffixOperation(FString& OutOriginalName)
 {
-	if (bRemoveOldSuffixSection)
+	if (CanApplyRemoveSuffixNumbers())
 	{
-		if (CanApplyRemoveSuffixNumbers())
-		{
-			ApplyRemoveSuffixNumbers(OutOriginalName);
-		}
+		ApplyRemoveSuffixNumbers(OutOriginalName);
+	}
 
+	if (bRemoveOldSuffixSection)
+	{
 		if (CanApplyRemoveSuffixSeparatorOperation())
 		{
 			ApplyRemoveSuffixSeparatorOperation(OutOriginalName);
diff --git a/Engine/Plugins/Experimental/AdvancedRenamer/Source/AdvancedRenamer/Private/AdvancedRenamerSections/AdvancedRenamerSearchAndReplaceSection.cpp b/Engine/Plugins/Experimental/AdvancedRenamer/Source/AdvancedRenamer/Private/AdvancedRenamerSections/AdvancedRenamerSearchAndReplaceSection.cpp
--- a/Engine/Plugins/Experimental/AdvancedRenamer/Source/AdvancedRenamer/Private/AdvancedRenamerSections/AdvancedRenamerSearchAndReplaceSection.cpp
+++ b/Engine/Plugins/Experimental/AdvancedRenamer/Source/AdvancedRenamer/Private/AdvancedRenamerSections/AdvancedRenamerSearchAndReplaceSection.cpp
@@ -45,7 +45,7 @@ TSharedRef<SWidget> FAdvancedRenamerSearchAndReplaceSection::GetWidget()
 			[
 				SNew(STextBlock)
 				.Font(FAdvancedRenamerStyle::Get().GetFontStyle("AdvancedRenamer.Style.TitleFont"))
-				.Text(LOCTEXT("AR_SearchReplaceTitle", "Search and Replace"))
+				.Text(LOCTEXT("AR_SearchReplaceTitle", "Rename"))
 			]
 
 			+ SVerticalBox::Slot()
@@ -58,7 +58,7 @@ TSharedRef<SWidget> FAdvancedRenamerSearchAndReplaceSection::GetWidget()
 				.HAlign(HAlign_Left)
 				.VAlign(VAlign_Center)
 				.AutoWidth()
-				.Padding(NameWidgetPadding)
+				.Padding(FirstWidgetPadding)
 				[
 					SNew(SBox)
 					.HeightOverride(25.f)
@@ -122,24 +122,59 @@ TSharedRef<SWidget> FAdvancedRenamerSearchAndReplaceSection::GetWidget()
 			.Padding(SectionContentMiddleEntriesPadding)
 			.FillHeight(1.f)
 			[
-				SAssignNew(SearchReplaceSearchTextBox, SEditableTextBox)
-				.BackgroundColor(FStyleColors::Background.GetSpecifiedColor())
-				.Font(FAdvancedRenamerStyle::Get().GetFontStyle("AdvancedRenamer.Style.RegularFont"))
-				.HintText(LOCTEXT("AR_RegexSearchHint", "Search"))
-				.Text(this, &FAdvancedRenamerSearchAndReplaceSection::GetSearchText)
-				.OnTextChanged(this, &FAdvancedRenamerSearchAndReplaceSection::OnSearchTextChanged)
+				SNew(SHorizontalBox)
+
+				+ SHorizontalBox::Slot()
+				.Padding(FirstWidgetPadding)
+				.VAlign(VAlign_Center)
+				.AutoWidth()
+				[
+					SNew(SBox)
+					.WidthOverride(65.f)
+					[
+						SNew(STextBlock)
+						.Font(FAdvancedRenamerStyle::Get().GetFontStyle("AdvancedRenamer.Style.RegularFont"))
+						.Text(LOCTEXT("AR_SearchLabel", "Search"))
+					]
+				]
+
+				+ SHorizontalBox::Slot()
+				.FillWidth(1.f)
+				[
+					SAssignNew(SearchReplaceSearchTextBox, SEditableTextBox)
+					.Font(FAdvancedRenamerStyle::Get().GetFontStyle("AdvancedRenamer.Style.RegularFont"))
+					.HintText(LOCTEXT("AR_RegexSearchHint", "Search (Optional)"))
+					.Text(this, &FAdvancedRenamerSearchAndReplaceSection::GetSearchText)
+					.OnTextChanged(this, &FAdvancedRenamerSearchAndReplaceSection::OnSearchTextChanged)
+				]
 			]
 
+			// Rename/Replace text
 			+ SVerticalBox::Slot()
 			.Padding(SectionContentMiddleEntriesPadding)
 			.FillHeight(1.f)
 			[
-				SAssignNew(SearchReplaceReplaceTextBox, SEditableTextBox)
-				.BackgroundColor(FStyleColors::Background.GetSpecifiedColor())
-				.Font(FAdvancedRenamerStyle::Get().GetFontStyle("AdvancedRenamer.Style.RegularFont"))
-				.HintText(LOCTEXT("AR_RegexReplaceHint", "Rename To"))
-				.Text(this, &FAdvancedRenamerSearchAndReplaceSection::GetReplaceText)
-				.OnTextChanged(this, &FAdvancedRenamerSearchAndReplaceSection::OnReplaceTextChanged)
+				SNew(SHorizontalBox)
+
+				+ SHorizontalBox::Slot()
+				.Padding(FirstWidgetPadding)
+				.VAlign(VAlign_Center)
+				.AutoWidth()
+				[
+					SNew(STextBlock)
+					.Font(FAdvancedRenamerStyle::Get().GetFontStyle("AdvancedRenamer.Style.RegularFont"))
+					.Text(LOCTEXT("AR_RenameReplaceLabel", "Rename To"))
+				]
+
+				+ SHorizontalBox::Slot()
+				.FillWidth(1.f)
+				[
+					SAssignNew(SearchReplaceReplaceTextBox, SEditableTextBox)
+					.Font(FAdvancedRenamerStyle::Get().GetFontStyle("AdvancedRenamer.Style.RegularFont"))
+					.HintText(LOCTEXT("AR_RenameReplaceHint", "Replace"))
+					.Text(this, &FAdvancedRenamerSearchAndReplaceSection::GetReplaceText)
+					.OnTextChanged(this, &FAdvancedRenamerSearchAndReplaceSection::OnReplaceTextChanged)
+				]
 			]
 		];
 }
diff --git a/Engine/Plugins/Experimental/AdvancedRenamer/Source/AdvancedRenamer/Private/Slate/SAdvancedRenamerPanel.cpp b/Engine/Plugins/Experimental/AdvancedRenamer/Source/AdvancedRenamer/Private/Slate/SAdvancedRenamerPanel.cpp
--- a/Engine/Plugins/Experimental/AdvancedRenamer/Source/AdvancedRenamer/Private/Slate/SAdvancedRenamerPanel.cpp
+++ b/Engine/Plugins/Experimental/AdvancedRenamer/Source/AdvancedRenamer/Private/Slate/SAdvancedRenamerPanel.cpp
@@ -118,7 +118,6 @@ void SAdvancedRenamerPanel::Construct(const FArguments& InArgs, const TSharedRef
 				[
 					SNew(SSplitter)
 					.PhysicalSplitterHandleSize(2.f)
-					.MinimumSlotHeight(499.5f)
 					.Style(&FAdvancedRenamerStyle::Get().GetWidgetStyle<FSplitterStyle>("AdvancedRenamer.Style.Splitter"))
 
 					+ SSplitter::Slot()
diff --git a/Engine/Plugins/Experimental/AdvancedRenamer/Source/AdvancedRenamer/Private/Utils/AdvancedRenamerSlateUtils.h b/Engine/Plugins/Experimental/AdvancedRenamer/Source/AdvancedRenamer/Private/Utils/AdvancedRenamerSlateUtils.h
--- a/Engine/Plugins/Experimental/AdvancedRenamer/Source/AdvancedRenamer/Private/Utils/AdvancedRenamerSlateUtils.h
+++ b/Engine/Plugins/Experimental/AdvancedRenamer/Source/AdvancedRenamer/Private/Utils/AdvancedRenamerSlateUtils.h
@@ -15,12 +15,12 @@ namespace AdvancedRenamerSlateUtils::Default
 {
 	static FName OriginalNameColumnName = TEXT("OriginalName");
 	static FName NewNameColumnName = TEXT("NewName");
-	static FMargin NameWidgetPadding = FMargin(0.f, 0.f, 4.f, 0.f);
-	static FMargin ValueWidgetPadding = FMargin(4.f, 0.f, 0.f, 0.f);
+	static FMargin FirstWidgetPadding = FMargin(0.f, 0.f, 4.f, 0.f);
+	static FMargin LastWidgetPadding = FMargin(4.f, 0.f, 0.f, 0.f);
+	static FMargin MiddleWidgetPadding = FMargin(4.f, 0.f, 4.f, 0.f);
 	static FMargin SectionContentFirstEntryPadding = FMargin(8.f, 8.f);
 	static FMargin SectionContentMiddleEntriesPadding = FMargin(8.f, 0.f, 8.f, 8.f);
-	static FMargin VerticalAddNumberPadding = FMargin(0.f, 4.f);
-	static FMargin AddNumberStepPadding = FMargin(8.f, 0.f, 0.f, 0.f);
+	static FMargin VerticalPadding = FMargin(0.f, 4.f);
 	static FMargin ChangeCaseFirstButtonPadding = FMargin(0.f, 0.f, 4.f, 0.f);
 	static FMargin ChangeCaseMiddleButtonsPadding = FMargin(4.f, 0.f);
 	static FMargin ChangeCaseLastButtonPadding = FMargin(4.f, 0.f, 0.f, 0.f);
