[BatchRenamer] Support ascending and descending order

diff --git a/Engine/Plugins/Experimental/AdvancedRenamer/Source/AdvancedRenamer/Private/AdvancedRenamerSections/AdvancedRenamerNumberingSection.cpp b/Engine/Plugins/Experimental/AdvancedRenamer/Source/AdvancedRenamer/Private/AdvancedRenamerSections/AdvancedRenamerNumberingSection.cpp
--- a/Engine/Plugins/Experimental/AdvancedRenamer/Source/AdvancedRenamer/Private/AdvancedRenamerSections/AdvancedRenamerNumberingSection.cpp
+++ b/Engine/Plugins/Experimental/AdvancedRenamer/Source/AdvancedRenamer/Private/AdvancedRenamerSections/AdvancedRenamerNumberingSection.cpp
@@ -94,10 +94,11 @@ TSharedRef<SWidget> FAdvancedRenamerNumberingSection::GetWidget()
 				+ SHorizontalBox::Slot()
 				.Padding(LastWidgetPadding)
 				.VAlign(VAlign_Center)
-				.FillWidth(0.5f)
+				.HAlign(HAlign_Right)
 				[
 					SNew(SBox)
 					.HeightOverride(25.f)
+					.MinDesiredWidth(160.f)
 					[
 						SNew(SSegmentedControl<EAdvancedRenamerNumberingType>)
 						.TextStyle(&FAppStyle::Get().GetWidgetStyle<FTextBlockStyle>("NormalText"))
diff --git a/Engine/Plugins/Experimental/AdvancedRenamer/Source/AdvancedRenamer/Private/IAdvancedRenamer.h b/Engine/Plugins/Experimental/AdvancedRenamer/Source/AdvancedRenamer/Private/IAdvancedRenamer.h
--- a/Engine/Plugins/Experimental/AdvancedRenamer/Source/AdvancedRenamer/Private/IAdvancedRenamer.h
+++ b/Engine/Plugins/Experimental/AdvancedRenamer/Source/AdvancedRenamer/Private/IAdvancedRenamer.h
@@ -13,12 +13,20 @@ struct FAdvancedRenamerPreview
 		: Hash(InHash)
 		, OriginalName(InOriginalName)
 		, NewName(TEXT(""))
+		, OriginalNameForSort(InOriginalName)
 	{
 	}
 
+	/** Get the OriginalName FName for sort purposes */
+	FName GetNameForSort() const { return OriginalNameForSort; }
+	
+public:
 	int32 Hash;
 	const FString OriginalName;
 	FString NewName;
+
+private:
+	const FName OriginalNameForSort;
 };
 
 /**
diff --git a/Engine/Plugins/Experimental/AdvancedRenamer/Source/AdvancedRenamer/Private/Slate/SAdvancedRenamerPanel.cpp b/Engine/Plugins/Experimental/AdvancedRenamer/Source/AdvancedRenamer/Private/Slate/SAdvancedRenamerPanel.cpp
--- a/Engine/Plugins/Experimental/AdvancedRenamer/Source/AdvancedRenamer/Private/Slate/SAdvancedRenamerPanel.cpp
+++ b/Engine/Plugins/Experimental/AdvancedRenamer/Source/AdvancedRenamer/Private/Slate/SAdvancedRenamerPanel.cpp
@@ -94,6 +94,9 @@ void SAdvancedRenamerPanel::Construct(const FArguments& InArgs, const TSharedRef
 {
 	Renamer = InRenamer;
 
+	SortMode = EColumnSortMode::None;
+	PreviewList = InRenamer->GetPreviews();
+
 	CommandList = MakeShared<FUICommandList>();
 	CommandList->MapAction(
 		FGenericCommands::Get().Delete,
@@ -323,7 +326,7 @@ TSharedRef<SWidget> SAdvancedRenamerPanel::CreateRenamePreview()
 				[
 					SAssignNew(RenamePreviewList, SListView<TSharedPtr<FAdvancedRenamerPreview>>)
 					.ListViewStyle(&FAdvancedRenamerStyle::Get().GetWidgetStyle<FTableViewStyle>("AdvancedRenamer.Style.ListView"))
-					.ListItemsSource(&Renamer->GetPreviews())
+					.ListItemsSource(&PreviewList)
 					.OnGenerateRow(this, &SAdvancedRenamerPanel::OnGenerateRowForList)
 					.HeaderRow(
 
@@ -333,6 +336,8 @@ TSharedRef<SWidget> SAdvancedRenamerPanel::CreateRenamePreview()
 						+ SHeaderRow::Column(OriginalNameColumnName)
 						.HeaderContentPadding(FMargin(8.f, 2.f))
 						.DefaultLabel(LOCTEXT("AR_Old", "Old"))
+						.SortMode(this, &SAdvancedRenamerPanel::GetColumnSortMode)
+						.OnSort(this, &SAdvancedRenamerPanel::OnColumnSortModeChanged)
 						.FillWidth(0.5f)
 
 						+ SHeaderRow::Column(NewNameColumnName)
@@ -442,8 +447,46 @@ void SAdvancedRenamerPanel::Tick(const FGeometry& InAllottedGeometry, const doub
 	}
 }
 
-TSharedRef<ITableRow> SAdvancedRenamerPanel::OnGenerateRowForList(TSharedPtr<FAdvancedRenamerPreview> InItem, 
-	const TSharedRef<STableViewBase>& InOwnerTable)
+EColumnSortMode::Type SAdvancedRenamerPanel::GetColumnSortMode() const
+{
+	return SortMode;
+}
+
+void SAdvancedRenamerPanel::OnColumnSortModeChanged(EColumnSortPriority::Type InSortPriority, const FName& InName, EColumnSortMode::Type InSortMode)
+{
+	if (SortMode == EColumnSortMode::Descending && InSortMode == EColumnSortMode::Ascending)
+	{
+		SortMode = EColumnSortMode::None;
+	}
+	else
+	{
+		SortMode = InSortMode;
+	}
+
+	if (Renamer.IsValid())
+	{
+		PreviewList = Renamer->GetPreviews();
+		if (SortMode != EColumnSortMode::None)
+		{
+			EColumnSortMode::Type TempSortMode = SortMode;
+			auto ComparePreviewList = [TempSortMode](const TSharedPtr<FAdvancedRenamerPreview>& A, const TSharedPtr<FAdvancedRenamerPreview>& B)
+			{
+				if (!A.IsValid() || !B.IsValid())
+				{
+					return false;
+				}
+
+				const bool CompareResult = A->GetNameForSort().Compare(B->GetNameForSort()) <= 0;
+				return TempSortMode == EColumnSortMode::Ascending ? CompareResult : !CompareResult;
+			};
+
+			PreviewList.Sort(ComparePreviewList);
+		}
+		RenamePreviewList->RequestListRefresh();
+	}
+}
+
+TSharedRef<ITableRow> SAdvancedRenamerPanel::OnGenerateRowForList(TSharedPtr<FAdvancedRenamerPreview> InItem, const TSharedRef<STableViewBase>& InOwnerTable)
 {
 	return SNew(SAdvancedRenamerPreviewListRow, SharedThis(this), InOwnerTable, InItem);
 }
diff --git a/Engine/Plugins/Experimental/AdvancedRenamer/Source/AdvancedRenamer/Private/Slate/SAdvancedRenamerPanel.h b/Engine/Plugins/Experimental/AdvancedRenamer/Source/AdvancedRenamer/Private/Slate/SAdvancedRenamerPanel.h
--- a/Engine/Plugins/Experimental/AdvancedRenamer/Source/AdvancedRenamer/Private/Slate/SAdvancedRenamerPanel.h
+++ b/Engine/Plugins/Experimental/AdvancedRenamer/Source/AdvancedRenamer/Private/Slate/SAdvancedRenamerPanel.h
@@ -4,13 +4,13 @@
 
 #include "Input/Reply.h"
 #include "Widgets/SCompoundWidget.h"
+#include "Widgets/Views/SHeaderRow.h"
 #include "Templates/SharedPointerFwd.h"
 
 class FUICommandList;
 class IAdvancedRenamer;
 class ITableRow;
 class SBox;
-class SHeaderRow;
 template <typename InItemType>
 class SListView;
 class STableViewBase;
@@ -62,6 +62,12 @@ private:
 	virtual void Tick(const FGeometry& InAllottedGeometry, const double InCurrentTime, const float InDeltaTime) override;
 	//~ End SWidget Interface
 
+	/** Get the current Column SortMode */
+	EColumnSortMode::Type GetColumnSortMode() const;
+
+	/** Called when the SortMode change */
+	void OnColumnSortModeChanged(EColumnSortPriority::Type InSortPriority, const FName& InName, EColumnSortMode::Type InSortMode);
+	
 	/** Callback to generate ListBoxRows */
 	TSharedRef<ITableRow> OnGenerateRowForList(TSharedPtr<FAdvancedRenamerPreview> InItem, const TSharedRef<STableViewBase>& InOwnerTable);
 
@@ -87,6 +93,12 @@ private:
 	/** Min update frequency, used in the tick to avoid updating the renamer too often */
 	static constexpr double MinUpdateFrequency = 0.1f;
 
+	/** Current SortMode for the List */
+	EColumnSortMode::Type SortMode;
+
+	/** Preview list used for the ListView */
+	TArray<TSharedPtr<FAdvancedRenamerPreview>> PreviewList;
+
 	/** Renamer instance */
 	TSharedPtr<IAdvancedRenamer> Renamer;
 
