[BatchRenamer] fix issue with AppScale

diff --git a/Engine/Plugins/Experimental/AdvancedRenamer/Source/AdvancedRenamer/Private/AdvancedRenamerModule.cpp b/Engine/Plugins/Experimental/AdvancedRenamer/Source/AdvancedRenamer/Private/AdvancedRenamerModule.cpp
--- a/Engine/Plugins/Experimental/AdvancedRenamer/Source/AdvancedRenamer/Private/AdvancedRenamerModule.cpp
+++ b/Engine/Plugins/Experimental/AdvancedRenamer/Source/AdvancedRenamer/Private/AdvancedRenamerModule.cpp
@@ -28,15 +28,17 @@ namespace UE::AdvancedRenamer::Private
 {
 	TSharedRef<SWindow> CreateAdvancedRenamerWindow()
 	{
+		// Take into account the AppScale when spawning the window, this has to be done since SWindow do not take the AppScale into consideration for the Min/Max Width/Height
+		float AppScale = FSlateApplication::Get().GetApplicationScale();
 		return SNew(SWindow)
 			.Title(LOCTEXT("AdvancedRenameWindow", "Batch Renamer"))
-			.ClientSize(FVector2D(731, 603))
+			.ClientSize(FVector2D(731 * AppScale, 603 * AppScale))
 			.SizingRule(ESizingRule::UserSized)
 			.SupportsMaximize(false)
 			.SupportsMinimize(false)
-			.MinWidth(737)
-			.MaxHeight(626.f)
-			.MinHeight(626.f);
+			.MinWidth(737.f * AppScale)
+			.MaxHeight(626.f * AppScale)
+			.MinHeight(626.f * AppScale);
 	}
 }
 
