Advanced Renamer can now also rename localized variants alongside source assets

diff --git a/Engine/Plugins/Experimental/AdvancedRenamer/Source/AdvancedRenamer/Private/AdvancedRenamer.cpp b/Engine/Plugins/Experimental/AdvancedRenamer/Source/AdvancedRenamer/Private/AdvancedRenamer.cpp
--- a/Engine/Plugins/Experimental/AdvancedRenamer/Source/AdvancedRenamer/Private/AdvancedRenamer.cpp
+++ b/Engine/Plugins/Experimental/AdvancedRenamer/Source/AdvancedRenamer/Private/AdvancedRenamer.cpp
@@ -124,30 +124,21 @@ bool FAdvancedRenamer::Execute()
 	}
 
 	const int32 Count = Previews.Num();
-	bool bAllSuccess = true;
 
+	bool bAllSuccess = Provider->BeginRename();
 	for (int32 Index = 0; Index < Count; ++Index)
 	{
-		if (!Previews[Index].IsValid())
+		if (!Previews[Index].IsValid()
+			|| !IsValidIndex(Index)
+			|| Previews[Index]->NewName.IsEmpty())
 		{
 			continue;
 		}
 
-		if (!IsValidIndex(Index))
-		{
-			continue;
-		}
-
-		if (Previews[Index]->NewName.IsEmpty())
-		{
-			continue;
-		}
-
-		if (!ExecuteRename(Index, Previews[Index]->NewName))
-		{
-			bAllSuccess = false;
-		}
+		bAllSuccess &= Provider->PrepareRename(Index, Previews[Index]->NewName);
 	}
+	bAllSuccess &= Provider->ExecuteRename();
+	bAllSuccess &= Provider->EndRename();
 
 	MarkClean();
 
@@ -217,7 +208,22 @@ bool FAdvancedRenamer::CanRename(int32 InIndex) const
 	return Provider->CanRename(InIndex);
 }
 
-bool FAdvancedRenamer::ExecuteRename(int32 InIndex, const FString& InNewName)
+bool FAdvancedRenamer::BeginRename()
 {
-	return Provider->ExecuteRename(InIndex, InNewName);
+	return Provider->BeginRename();
 }
+
+bool FAdvancedRenamer::PrepareRename(int32 InIndex, const FString& InNewName)
+{
+	return Provider->PrepareRename(InIndex, InNewName);
+}
+
+bool FAdvancedRenamer::ExecuteRename()
+{
+	return Provider->ExecuteRename();
+}
+
+bool FAdvancedRenamer::EndRename()
+{
+	return Provider->EndRename();
+}
\ No newline at end of file
diff --git a/Engine/Plugins/Experimental/AdvancedRenamer/Source/AdvancedRenamer/Private/AdvancedRenamer.h b/Engine/Plugins/Experimental/AdvancedRenamer/Source/AdvancedRenamer/Private/AdvancedRenamer.h
--- a/Engine/Plugins/Experimental/AdvancedRenamer/Source/AdvancedRenamer/Private/AdvancedRenamer.h
+++ b/Engine/Plugins/Experimental/AdvancedRenamer/Source/AdvancedRenamer/Private/AdvancedRenamer.h
@@ -34,7 +34,11 @@ public:
 	virtual FString GetOriginalName(int32 InIndex) const override;
 	virtual bool RemoveIndex(int32 InIndex) override;
 	virtual bool CanRename(int32 InIndex) const override;
-	virtual bool ExecuteRename(int32 InIndex, const FString& InNewName) override;
+
+	virtual bool BeginRename() override;
+	virtual bool PrepareRename(int32 InIndex, const FString& InNewName) override;
+	virtual bool ExecuteRename() override;
+	virtual bool EndRename() override;
 	//~ End IAdvancedRenamerProvider
 
 private:
diff --git a/Engine/Plugins/Experimental/AdvancedRenamer/Source/AdvancedRenamer/Private/Providers/AdvancedRenamerActorProvider.cpp b/Engine/Plugins/Experimental/AdvancedRenamer/Source/AdvancedRenamer/Private/Providers/AdvancedRenamerActorProvider.cpp
--- a/Engine/Plugins/Experimental/AdvancedRenamer/Source/AdvancedRenamer/Private/Providers/AdvancedRenamerActorProvider.cpp
+++ b/Engine/Plugins/Experimental/AdvancedRenamer/Source/AdvancedRenamer/Private/Providers/AdvancedRenamerActorProvider.cpp
@@ -96,7 +96,13 @@ bool FAdvancedRenamerActorProvider::CanRename(int32 InIndex) const
 	return true;
 }
 
-bool FAdvancedRenamerActorProvider::ExecuteRename(int32 InIndex, const FString& InNewName)
+bool FAdvancedRenamerActorProvider::BeginRename()
+{
+	ActorToNewNameList.Reserve(Num());
+	return true;
+}
+
+bool FAdvancedRenamerActorProvider::PrepareRename(int32 InIndex, const FString& InNewName)
 {
 	AActor* Actor = GetActor(InIndex);
 
@@ -105,6 +111,21 @@ bool FAdvancedRenamerActorProvider::ExecuteRename(int32 InIndex, const FString&
 		return false;
 	}
 
-	Actor->SetActorLabel(InNewName, /* bMarkDirty */ true);
+	ActorToNewNameList.Add({ Actor , InNewName });
+	return true;
+}
+
+bool FAdvancedRenamerActorProvider::ExecuteRename()
+{
+	for (const TTuple<AActor*, FString>& ActorToNewName : ActorToNewNameList)
+	{
+		ActorToNewName.Key->SetActorLabel(ActorToNewName.Value, /* bMarkDirty */ true);
+	}
+	return true;
+}
+
+bool FAdvancedRenamerActorProvider::EndRename()
+{
+	ActorToNewNameList.Empty();
 	return true;
 }
diff --git a/Engine/Plugins/Experimental/AdvancedRenamer/Source/AdvancedRenamer/Private/Providers/AdvancedRenamerActorProvider.h b/Engine/Plugins/Experimental/AdvancedRenamer/Source/AdvancedRenamer/Private/Providers/AdvancedRenamerActorProvider.h
--- a/Engine/Plugins/Experimental/AdvancedRenamer/Source/AdvancedRenamer/Private/Providers/AdvancedRenamerActorProvider.h
+++ b/Engine/Plugins/Experimental/AdvancedRenamer/Source/AdvancedRenamer/Private/Providers/AdvancedRenamerActorProvider.h
@@ -29,8 +29,13 @@ protected:
 	virtual FString GetOriginalName(int32 InIndex) const override;
 	virtual bool RemoveIndex(int32 InIndex) override;
 	virtual bool CanRename(int32 InIndex) const override;
-	virtual bool ExecuteRename(int32 InIndex, const FString& InNewName) override;
+
+	virtual bool BeginRename() override;
+	virtual bool PrepareRename(int32 InIndex, const FString& InNewName) override;
+	virtual bool ExecuteRename() override;
+	virtual bool EndRename() override;
 	//~ End IAdvancedRenamerProvider
 
 	TArray<TWeakObjectPtr<AActor>> ActorList;
+	TArray<TTuple<AActor*, FString>> ActorToNewNameList;
 };
diff --git a/Engine/Plugins/Experimental/AdvancedRenamer/Source/AdvancedRenamer/Private/Providers/AdvancedRenamerAssetProvider.cpp b/Engine/Plugins/Experimental/AdvancedRenamer/Source/AdvancedRenamer/Private/Providers/AdvancedRenamerAssetProvider.cpp
--- a/Engine/Plugins/Experimental/AdvancedRenamer/Source/AdvancedRenamer/Private/Providers/AdvancedRenamerAssetProvider.cpp
+++ b/Engine/Plugins/Experimental/AdvancedRenamer/Source/AdvancedRenamer/Private/Providers/AdvancedRenamerAssetProvider.cpp
@@ -99,7 +99,13 @@ bool FAdvancedRenamerAssetProvider::CanRename(int32 InIndex) const
 	return true;
 }
 
-bool FAdvancedRenamerAssetProvider::ExecuteRename(int32 InIndex, const FString& InNewName)
+bool FAdvancedRenamerAssetProvider::BeginRename()
+{
+	AssetRenameDataList.Reserve(Num());
+	return true;
+}
+
+bool FAdvancedRenamerAssetProvider::PrepareRename(int32 InIndex, const FString& InNewName)
 {
 	UObject* Asset = GetAsset(InIndex);
 
@@ -111,8 +117,25 @@ bool FAdvancedRenamerAssetProvider::ExecuteRename(int32 InIndex, const FString&
 	FString PackagePath = Asset->GetPathName();
 	PackagePath = FPaths::GetPath(PackagePath);
 
+	constexpr bool bOnlyFixSoftReferences = false;
+	constexpr bool bAlsoRenameLocalizedVariants = true;
+	AssetRenameDataList.Add(FAssetRenameData(Asset, PackagePath, InNewName, bOnlyFixSoftReferences, bAlsoRenameLocalizedVariants));
+
+	return true;
+}
+
+bool FAdvancedRenamerAssetProvider::ExecuteRename()
+{
 	IAssetTools& AssetTools = FModuleManager::GetModuleChecked<FAssetToolsModule>("AssetTools").Get();
-	return AssetTools.RenameAssets({FAssetRenameData(Asset, PackagePath, InNewName)});
+
+	constexpr bool bAutoCheckout = false;
+	return AssetTools.RenameAssetsWithDialog(AssetRenameDataList, bAutoCheckout) != EAssetRenameResult::Failure;
+}
+
+bool FAdvancedRenamerAssetProvider::EndRename()
+{
+	AssetRenameDataList.Empty();
+	return true;
 }
 
 #undef LOCTEXT_NAMESPACE
diff --git a/Engine/Plugins/Experimental/AdvancedRenamer/Source/AdvancedRenamer/Private/Providers/AdvancedRenamerAssetProvider.h b/Engine/Plugins/Experimental/AdvancedRenamer/Source/AdvancedRenamer/Private/Providers/AdvancedRenamerAssetProvider.h
--- a/Engine/Plugins/Experimental/AdvancedRenamer/Source/AdvancedRenamer/Private/Providers/AdvancedRenamerAssetProvider.h
+++ b/Engine/Plugins/Experimental/AdvancedRenamer/Source/AdvancedRenamer/Private/Providers/AdvancedRenamerAssetProvider.h
@@ -7,6 +7,7 @@
 #include "Containers/Array.h"
 
 struct FAssetData;
+struct FAssetRenameData;
 
 class FAdvancedRenamerAssetProvider : public IAdvancedRenamerProvider
 {
@@ -27,8 +28,13 @@ protected:
 	virtual FString GetOriginalName(int32 InIndex) const override;
 	virtual bool RemoveIndex(int32 InIndex) override;
 	virtual bool CanRename(int32 InIndex) const override;
-	virtual bool ExecuteRename(int32 InIndex, const FString& InNewName) override;
+
+	virtual bool BeginRename() override;
+	virtual bool PrepareRename(int32 InIndex, const FString& InNewName) override;
+	virtual bool ExecuteRename() override;
+	virtual bool EndRename() override;
 	//~ End IAdvancedRenamerProvider
 
 	TArray<FAssetData> AssetList;
+	TArray<FAssetRenameData> AssetRenameDataList;
 };
diff --git a/Engine/Plugins/Experimental/AdvancedRenamer/Source/AdvancedRenamer/Private/Providers/AdvancedRenamerObjectProvider.cpp b/Engine/Plugins/Experimental/AdvancedRenamer/Source/AdvancedRenamer/Private/Providers/AdvancedRenamerObjectProvider.cpp
--- a/Engine/Plugins/Experimental/AdvancedRenamer/Source/AdvancedRenamer/Private/Providers/AdvancedRenamerObjectProvider.cpp
+++ b/Engine/Plugins/Experimental/AdvancedRenamer/Source/AdvancedRenamer/Private/Providers/AdvancedRenamerObjectProvider.cpp
@@ -96,7 +96,13 @@ bool FAdvancedRenamerObjectProvider::CanRename(int32 InIndex) const
 	return true;
 }
 
-bool FAdvancedRenamerObjectProvider::ExecuteRename(int32 InIndex, const FString& InNewName)
+bool FAdvancedRenamerObjectProvider::BeginRename()
+{
+	ObjectToNewNameList.Reserve(Num());
+	return true;
+}
+
+bool FAdvancedRenamerObjectProvider::PrepareRename(int32 InIndex, const FString& InNewName)
 {
 	UObject* Object = GetObject(InIndex);
 
@@ -105,11 +111,31 @@ bool FAdvancedRenamerObjectProvider::ExecuteRename(int32 InIndex, const FString&
 		return false;
 	}
 
-	if (Object->Rename(*InNewName, nullptr, REN_Test))
+	ObjectToNewNameList.Add({ Object, InNewName });
+
+	return true;
+}
+
+bool FAdvancedRenamerObjectProvider::ExecuteRename()
+{
+	bool bAllSuccess = true;
+	for (const TTuple<UObject*, FString>& ObjectToNewName : ObjectToNewNameList)
 	{
-		Object->Rename(*InNewName);
-		return true;
+		// We have to run the REN_Test here to prevent the other renamed objects from making this rename fail
+		if (ObjectToNewName.Key->Rename(*ObjectToNewName.Value, nullptr, REN_Test))
+		{
+			bAllSuccess &= ObjectToNewName.Key->Rename(*ObjectToNewName.Value);
+		}
+		else
+		{
+			bAllSuccess = false;
+		}
 	}
+	return bAllSuccess;
+}
 
-	return false;
+bool FAdvancedRenamerObjectProvider::EndRename()
+{
+	ObjectToNewNameList.Empty();
+	return true;
 }
diff --git a/Engine/Plugins/Experimental/AdvancedRenamer/Source/AdvancedRenamer/Private/Providers/AdvancedRenamerObjectProvider.h b/Engine/Plugins/Experimental/AdvancedRenamer/Source/AdvancedRenamer/Private/Providers/AdvancedRenamerObjectProvider.h
--- a/Engine/Plugins/Experimental/AdvancedRenamer/Source/AdvancedRenamer/Private/Providers/AdvancedRenamerObjectProvider.h
+++ b/Engine/Plugins/Experimental/AdvancedRenamer/Source/AdvancedRenamer/Private/Providers/AdvancedRenamerObjectProvider.h
@@ -29,8 +29,13 @@ protected:
 	virtual FString GetOriginalName(int32 InIndex) const override;
 	virtual bool RemoveIndex(int32 InIndex) override;
 	virtual bool CanRename(int32 InIndex) const override;
-	virtual bool ExecuteRename(int32 InIndex, const FString& InNewName) override;
+
+	virtual bool BeginRename() override;
+	virtual bool PrepareRename(int32 InIndex, const FString& InNewName) override;
+	virtual bool ExecuteRename() override;
+	virtual bool EndRename() override;
 	//~ End IAdvancedRenamerProvider
 
 	TArray<TWeakObjectPtr<UObject>> ObjectList;
+	TArray<TTuple<UObject*, FString>> ObjectToNewNameList;
 };
diff --git a/Engine/Plugins/Experimental/AdvancedRenamer/Source/AdvancedRenamer/Public/Providers/IAdvancedRenamerProvider.h b/Engine/Plugins/Experimental/AdvancedRenamer/Source/AdvancedRenamer/Public/Providers/IAdvancedRenamerProvider.h
--- a/Engine/Plugins/Experimental/AdvancedRenamer/Source/AdvancedRenamer/Public/Providers/IAdvancedRenamerProvider.h
+++ b/Engine/Plugins/Experimental/AdvancedRenamer/Source/AdvancedRenamer/Public/Providers/IAdvancedRenamerProvider.h
@@ -17,7 +17,11 @@ public:
 	virtual FString GetOriginalName(int32 InIndex) const = 0;
 	virtual bool RemoveIndex(int32 InIndex) = 0;
 	virtual bool CanRename(int32 InIndex) const = 0;
-	virtual bool ExecuteRename(int32 InIndex, const FString& InNewName) = 0;
+
+	virtual bool BeginRename() = 0;
+	virtual bool PrepareRename(int32 InIndex, const FString& InNewName) = 0;
+	virtual bool ExecuteRename() = 0;
+	virtual bool EndRename() = 0;
 
 	virtual int32 FindHash(int32 InInHash) const
 	{
diff --git a/Engine/Source/Developer/AssetTools/Private/AssetRenameManager.cpp b/Engine/Source/Developer/AssetTools/Private/AssetRenameManager.cpp
--- a/Engine/Source/Developer/AssetTools/Private/AssetRenameManager.cpp
+++ b/Engine/Source/Developer/AssetTools/Private/AssetRenameManager.cpp
@@ -343,6 +343,27 @@ bool FAssetRenameManager::RenameAssetsAndVariants(const TArray<FAssetRenameData>
 			}
 			return false;
 		}
+		
+		// Display newly added rename data if possible
+		if (bWithDialog)
+		{
+			TArray<FText> NewAssetsAdded;
+			for (const FAssetRenameData& AssetRenameData : AssetsAndVariants)
+			{
+				if (!InAssetsRenameData.ContainsByPredicate([&AssetRenameData](const FAssetRenameData& InAssetRenameData) { return AssetRenameData.Asset == InAssetRenameData.Asset; }))
+				{
+					const FString& AssetNameStr = AssetRenameData.Asset->GetOuter()->GetName();
+					UE_LOG(LogAssetTools, Display, TEXT("Also trying to rename: %s"), *AssetNameStr);
+					NewAssetsAdded.Add(FText::AsCultureInvariant(AssetNameStr));
+				}
+			}
+			if (!NewAssetsAdded.IsEmpty())
+			{
+				LocalizedAssetTools->OpenLocalizedVariantsListMessageDialog(LOCTEXT("AddedLocalizedVariantsHeader", "Renaming localized variants too"),
+					LOCTEXT("AddedLocalizedVariantsMessage", "The following localized variants (or source assets) are also going to be renamed alongside the selected assets."),
+					NewAssetsAdded);
+			}
+		}
 	}
 
 	// We now have a full list of Assets to rename and their variants if applicable. Let's continue the renaming process.
diff --git a/Engine/Source/Developer/AssetTools/Private/LocalizedAssetTools.cpp b/Engine/Source/Developer/AssetTools/Private/LocalizedAssetTools.cpp
--- a/Engine/Source/Developer/AssetTools/Private/LocalizedAssetTools.cpp
+++ b/Engine/Source/Developer/AssetTools/Private/LocalizedAssetTools.cpp
@@ -223,7 +223,7 @@ void FLocalizedAssetTools::OpenFilesInRevisionControlRequiredDialog(const TArray
 
 void FLocalizedAssetTools::OpenLocalizedVariantsListMessageDialog(const FText& Header, const FText& Message, const TArray<FText>& FileList) const
 {
-	SFileListReportDialog::OpenDialog(Header, Message, FileList);
+	SFileListReportDialog::OpenDialog(Header, Message, FileList, true);
 }
 
 const FText& FLocalizedAssetTools::GetRevisionControlIsNotAvailableWarningText() const
diff --git a/Engine/Source/Developer/AssetTools/Private/SFileListReportDialog.cpp b/Engine/Source/Developer/AssetTools/Private/SFileListReportDialog.cpp
--- a/Engine/Source/Developer/AssetTools/Private/SFileListReportDialog.cpp
+++ b/Engine/Source/Developer/AssetTools/Private/SFileListReportDialog.cpp
@@ -65,9 +65,9 @@ void SFileListReportDialog::Construct(const FArguments& InArgs)
 		];
 }
 
-void SFileListReportDialog::OpenDialog(const FText& InTitle, const FText& InHeader, const TArray<FText>& InFiles)
+void SFileListReportDialog::OpenDialog(const FText& InTitle, const FText& InHeader, const TArray<FText>& InFiles, bool bOpenAsModal /*= false*/)
 {
-	TSharedRef<SWindow> RenameWindow = SNew(SWindow)
+	TSharedRef<SWindow> FileListReportWindow = SNew(SWindow)
 		.Title(InTitle)
 		.ClientSize(FVector2D(800, 400))
 		.SupportsMaximize(false)
@@ -77,14 +77,21 @@ void SFileListReportDialog::OpenDialog(const FText& InTitle, const FText& InHead
 		];
 
 	IMainFrameModule& MainFrameModule = FModuleManager::LoadModuleChecked<IMainFrameModule>(TEXT("MainFrame"));
-
+	
 	if (MainFrameModule.GetParentWindow().IsValid())
 	{
-		FSlateApplication::Get().AddWindowAsNativeChild(RenameWindow, MainFrameModule.GetParentWindow().ToSharedRef());
+		if (bOpenAsModal)
+		{
+			FSlateApplication::Get().AddModalWindow(FileListReportWindow, MainFrameModule.GetParentWindow().ToSharedRef());
+		}
+		else
+		{
+			FSlateApplication::Get().AddWindowAsNativeChild(FileListReportWindow, MainFrameModule.GetParentWindow().ToSharedRef());
+		}
 	}
 	else
 	{
-		FSlateApplication::Get().AddWindow(RenameWindow);
+		FSlateApplication::Get().AddWindow(FileListReportWindow);
 	}
 }
 
diff --git a/Engine/Source/Developer/AssetTools/Private/SFileListReportDialog.h b/Engine/Source/Developer/AssetTools/Private/SFileListReportDialog.h
--- a/Engine/Source/Developer/AssetTools/Private/SFileListReportDialog.h
+++ b/Engine/Source/Developer/AssetTools/Private/SFileListReportDialog.h
@@ -20,7 +20,7 @@ public:
 
 	void Construct(const FArguments& InArgs);
 
-	static void OpenDialog(const FText& InTitle, const FText& InHeader, const TArray<FText>& InFiles);
+	static void OpenDialog(const FText& InTitle, const FText& InHeader, const TArray<FText>& InFiles, bool bOpenAsModal = false);
 
 private:
 	TSharedRef<ITableRow> MakeListViewWidget(TSharedRef<FText> Item, const TSharedRef<STableViewBase>& OwnerTable);
