[BatchRenamer] window layout broken when opened on screens with OS Display Scaling different than 100%

diff --git a/Engine/Plugins/Experimental/AdvancedRenamer/Source/AdvancedRenamer/Private/AdvancedRenamerModule.cpp b/Engine/Plugins/Experimental/AdvancedRenamer/Source/AdvancedRenamer/Private/AdvancedRenamerModule.cpp
--- a/Engine/Plugins/Experimental/AdvancedRenamer/Source/AdvancedRenamer/Private/AdvancedRenamerModule.cpp
+++ b/Engine/Plugins/Experimental/AdvancedRenamer/Source/AdvancedRenamer/Private/AdvancedRenamerModule.cpp
@@ -28,17 +28,17 @@ namespace UE::AdvancedRenamer::Private
 {
 	TSharedRef<SWindow> CreateAdvancedRenamerWindow()
 	{
-		// Take into account the AppScale when spawning the window, this has to be done since SWindow do not take the AppScale into consideration for the Min/Max Width/Height
-		float AppScale = FSlateApplication::Get().GetApplicationScale();
+		// Workaround to make the AppScale and OS Zoom work at the same time, this will eventually be changed to add constraint to the window size
+		const float AppScale = FSlateApplication::Get().GetApplicationScale();
+		constexpr float MinWindowHeight = 589.f;
+		constexpr float MinWindowWidth = 730.f;
+		constexpr float TitleHeightOffset = 38.f;
+		constexpr float ContentWidthOffset = 6.f;
 		return SNew(SWindow)
 			.Title(LOCTEXT("AdvancedRenameWindow", "Batch Renamer"))
-			.ClientSize(FVector2D(731 * AppScale, 603 * AppScale))
-			.SizingRule(ESizingRule::UserSized)
-			.SupportsMaximize(false)
-			.SupportsMinimize(false)
-			.MinWidth(737.f * AppScale)
-			.MaxHeight(626.f * AppScale)
-			.MinHeight(626.f * AppScale);
+			.ClientSize(FVector2D((MinWindowWidth + ContentWidthOffset) * AppScale, (MinWindowHeight + TitleHeightOffset) * AppScale))
+			.MinHeight((MinWindowHeight + TitleHeightOffset) * AppScale)
+			.MinWidth((MinWindowWidth + ContentWidthOffset) * AppScale);
 	}
 }
 
diff --git a/Engine/Plugins/Experimental/AdvancedRenamer/Source/AdvancedRenamer/Private/AdvancedRenamerSections/AdvancedRenamerSearchAndReplaceSection.cpp b/Engine/Plugins/Experimental/AdvancedRenamer/Source/AdvancedRenamer/Private/AdvancedRenamerSections/AdvancedRenamerSearchAndReplaceSection.cpp
--- a/Engine/Plugins/Experimental/AdvancedRenamer/Source/AdvancedRenamer/Private/AdvancedRenamerSections/AdvancedRenamerSearchAndReplaceSection.cpp
+++ b/Engine/Plugins/Experimental/AdvancedRenamer/Source/AdvancedRenamer/Private/AdvancedRenamerSections/AdvancedRenamerSearchAndReplaceSection.cpp
@@ -120,7 +120,7 @@ TSharedRef<SWidget> FAdvancedRenamerSearchAndReplaceSection::GetWidget()
 			// Search text
 			+ SVerticalBox::Slot()
 			.Padding(SectionContentMiddleEntriesPadding)
-			.FillHeight(1.f)
+			.AutoHeight()
 			[
 				SNew(SHorizontalBox)
 
@@ -152,7 +152,7 @@ TSharedRef<SWidget> FAdvancedRenamerSearchAndReplaceSection::GetWidget()
 			// Rename/Replace text
 			+ SVerticalBox::Slot()
 			.Padding(SectionContentMiddleEntriesPadding)
-			.FillHeight(1.f)
+			.AutoHeight()
 			[
 				SNew(SHorizontalBox)
 
diff --git a/Engine/Plugins/Experimental/AdvancedRenamer/Source/AdvancedRenamer/Private/Slate/SAdvancedRenamerPanel.cpp b/Engine/Plugins/Experimental/AdvancedRenamer/Source/AdvancedRenamer/Private/Slate/SAdvancedRenamerPanel.cpp
--- a/Engine/Plugins/Experimental/AdvancedRenamer/Source/AdvancedRenamer/Private/Slate/SAdvancedRenamerPanel.cpp
+++ b/Engine/Plugins/Experimental/AdvancedRenamer/Source/AdvancedRenamer/Private/Slate/SAdvancedRenamerPanel.cpp
@@ -115,7 +115,6 @@ void SAdvancedRenamerPanel::Construct(const FArguments& InArgs, const TSharedRef
 		[
 			SNew(SBorder)
 			.BorderImage(FAppStyle::GetBrush("NoBorder"))
-			.Padding(2.f, 2.f)
 			[
 				SNew(SBox)
 				[
@@ -158,18 +157,9 @@ void SAdvancedRenamerPanel::Construct(const FArguments& InArgs, const TSharedRef
 
 		if (Index == 0)
 		{
-			// The first Section Search and Replace will fill the height by design
 			Padding = FMargin(0.f, 0.f, 0.f, 1.f);
-			LeftSideVerticalBox->AddSlot()
-				.Padding(Padding)
-				.FillHeight(1.f)
-				[
-					RegisteredSections[Index]->GetWidget()
-				];
-			continue;
 		}
-
-		if (Index == RegisteredSections.Num() - 1)
+		else if (Index == RegisteredSections.Num() - 1)
 		{
 			Padding = FMargin(0.f, 1.f, 0.f, 0.f);
 		}
