[BatchRenamer] numbering should be based on the ordering requested and not always on the initial order

diff --git a/Engine/Plugins/Experimental/AdvancedRenamer/Source/AdvancedRenamer/Private/AdvancedRenamer.cpp b/Engine/Plugins/Experimental/AdvancedRenamer/Source/AdvancedRenamer/Private/AdvancedRenamer.cpp
--- a/Engine/Plugins/Experimental/AdvancedRenamer/Source/AdvancedRenamer/Private/AdvancedRenamer.cpp
+++ b/Engine/Plugins/Experimental/AdvancedRenamer/Source/AdvancedRenamer/Private/AdvancedRenamer.cpp
@@ -7,16 +7,16 @@
 FAdvancedRenamer::FAdvancedRenamer(const TSharedRef<IAdvancedRenamerProvider>& InProvider)
 	: Provider(InProvider)
 {
-	int32 Count = Num();
-	check(Count > 0);
+	int32 RenamedEntitiesCount = Num();
+	check(RenamedEntitiesCount > 0);
 
-	for (int32 Index = 0; Index < Count; ++Index)
+	for (int32 Index = 0; Index < RenamedEntitiesCount; ++Index)
 	{
 		if (!CanRename(Index))
 		{
 			RemoveIndex(Index);
 			--Index;
-			--Count;
+			--RenamedEntitiesCount;
 			continue;
 		}
 
@@ -25,6 +25,7 @@ FAdvancedRenamer::FAdvancedRenamer(const TSharedRef<IAdvancedRenamerProvider>& I
 
 		Previews.Add(MakeShared<FAdvancedRenamerPreview>(Hash, OriginalName));
 	}
+	SortablePreviews = Previews;
 }
 
 const TSharedRef<IAdvancedRenamerProvider>& FAdvancedRenamer::GetProvider() const
@@ -32,19 +33,14 @@ const TSharedRef<IAdvancedRenamerProvider>& FAdvancedRenamer::GetProvider() cons
 	return Provider;
 }
 
-const TArray<TSharedPtr<FAdvancedRenamerPreview>>& FAdvancedRenamer::GetPreviews()
+TArray<TSharedPtr<FAdvancedRenamerPreview>>& FAdvancedRenamer::GetSortablePreviews()
 {
-	return Previews;
+	return SortablePreviews;
 }
 
-TSharedPtr<FAdvancedRenamerPreview> FAdvancedRenamer::GetPreview(int32 InIndex) const
+void FAdvancedRenamer::ResetSortablePreviews()
 {
-	if (Previews.IsValidIndex(InIndex))
-	{
-		return Previews[InIndex];
-	}
-
-	return nullptr;
+	SortablePreviews = Previews;
 }
 
 void FAdvancedRenamer::AddSection(FAdvancedRenamerExecuteSection InSection)
@@ -75,13 +71,13 @@ void FAdvancedRenamer::MarkClean()
 bool FAdvancedRenamer::UpdatePreviews()
 {
 	bHasRenames = false;
-	const int32 Count = Previews.Num();
+	const int32 Count = SortablePreviews.Num();
 
 	BeforeOperationsStartExecute();
 
 	for (int32 Index = 0; Index < Count; ++Index)
 	{
-		if (!Previews[Index].IsValid() || !IsValidIndex(Index))
+		if (!SortablePreviews[Index].IsValid() || !IsValidIndex(Index))
 		{
 			RemoveIndex(Index);
 			--Index;
@@ -89,14 +85,14 @@ bool FAdvancedRenamer::UpdatePreviews()
 		}
 
 		// Force recreation
-		Previews[Index]->NewName = ApplyRename(Previews[Index]->OriginalName);
+		SortablePreviews[Index]->NewName = ApplyRename(SortablePreviews[Index]->OriginalName);
 
-		if (Previews[Index]->NewName.IsEmpty())
+		if (SortablePreviews[Index]->NewName.IsEmpty())
 		{
 			continue;
 		}
 
-		if (GetOriginalName(Index).Equals(Previews[Index]->NewName))
+		if (GetOriginalName(Index).Equals(SortablePreviews[Index]->NewName))
 		{
 			continue;
 		}
@@ -123,19 +119,19 @@ bool FAdvancedRenamer::Execute()
 		}
 	}
 
-	const int32 Count = Previews.Num();
+	const int32 Count = SortablePreviews.Num();
 
 	bool bAllSuccess = Provider->BeginRename();
 	for (int32 Index = 0; Index < Count; ++Index)
 	{
-		if (!Previews[Index].IsValid()
+		if (!SortablePreviews[Index].IsValid()
 			|| !IsValidIndex(Index)
-			|| Previews[Index]->NewName.IsEmpty())
+			|| SortablePreviews[Index]->NewName.IsEmpty())
 		{
 			continue;
 		}
 
-		bAllSuccess &= Provider->PrepareRename(Index, Previews[Index]->NewName);
+		bAllSuccess &= Provider->PrepareRename(Index, SortablePreviews[Index]->NewName);
 	}
 	bAllSuccess &= Provider->ExecuteRename();
 	bAllSuccess &= Provider->EndRename();
@@ -195,12 +191,14 @@ bool FAdvancedRenamer::RemoveIndex(int32 InIndex)
 {
 	// Can fail during construction when indices that aren't renameable are removed from the provider before
 	// they are added to ListData.
-	if (Previews.IsValidIndex(InIndex))
+	int32 OriginalIndex = INDEX_NONE;
+	if (SortablePreviews.IsValidIndex(InIndex))
 	{
-		Previews.RemoveAt(InIndex);
+		OriginalIndex = Previews.IndexOfByKey(SortablePreviews[InIndex]);
+		SortablePreviews.RemoveAt(InIndex);
 	}
 
-	return Provider->RemoveIndex(InIndex);
+	return Provider->RemoveIndex(OriginalIndex);
 }
 
 bool FAdvancedRenamer::CanRename(int32 InIndex) const
diff --git a/Engine/Plugins/Experimental/AdvancedRenamer/Source/AdvancedRenamer/Private/AdvancedRenamer.h b/Engine/Plugins/Experimental/AdvancedRenamer/Source/AdvancedRenamer/Private/AdvancedRenamer.h
--- a/Engine/Plugins/Experimental/AdvancedRenamer/Source/AdvancedRenamer/Private/AdvancedRenamer.h
+++ b/Engine/Plugins/Experimental/AdvancedRenamer/Source/AdvancedRenamer/Private/AdvancedRenamer.h
@@ -15,8 +15,8 @@ public:
 
 	//~ Begin IAdvancedRenamer
 	virtual const TSharedRef<IAdvancedRenamerProvider>& GetProvider() const override;
-	virtual const TArray<TSharedPtr<FAdvancedRenamerPreview>>& GetPreviews() override;
-	virtual TSharedPtr<FAdvancedRenamerPreview> GetPreview(int32 InIndex) const override;
+	virtual TArray<TSharedPtr<FAdvancedRenamerPreview>>& GetSortablePreviews() override;
+	virtual void ResetSortablePreviews() override;
 	virtual void AddSection(FAdvancedRenamerExecuteSection InSection) override;
 	virtual bool HasRenames() const override;
 	virtual bool IsDirty() const override;
@@ -54,7 +54,10 @@ private:
 
 	/** Previews Name list */
 	TArray<TSharedPtr<FAdvancedRenamerPreview>> Previews;
-	
+
+	/** Sorted Previews Name list */
+	TArray<TSharedPtr<FAdvancedRenamerPreview>> SortablePreviews;
+
 	/** Renamer sections list */
 	TArray<FAdvancedRenamerExecuteSection> Sections;
 
diff --git a/Engine/Plugins/Experimental/AdvancedRenamer/Source/AdvancedRenamer/Private/IAdvancedRenamer.h b/Engine/Plugins/Experimental/AdvancedRenamer/Source/AdvancedRenamer/Private/IAdvancedRenamer.h
--- a/Engine/Plugins/Experimental/AdvancedRenamer/Source/AdvancedRenamer/Private/IAdvancedRenamer.h
+++ b/Engine/Plugins/Experimental/AdvancedRenamer/Source/AdvancedRenamer/Private/IAdvancedRenamer.h
@@ -40,10 +40,11 @@ public:
 
 	virtual const TSharedRef<IAdvancedRenamerProvider>& GetProvider() const = 0;
 
-	virtual const TArray<TSharedPtr<FAdvancedRenamerPreview>>& GetPreviews() = 0;
+	/** Get the sortable previews */
+	virtual TArray<TSharedPtr<FAdvancedRenamerPreview>>& GetSortablePreviews() = 0;
 
-	/** Returns the preview for the item at the given index. */
-	virtual TSharedPtr<FAdvancedRenamerPreview> GetPreview(int32 InIndex) const = 0;
+	/** Reset the order of the Previews to the original one */
+	virtual void ResetSortablePreviews() = 0;
 
 	/** Add a section to this Renamer */
 	virtual void AddSection(FAdvancedRenamerExecuteSection InSection) = 0;
diff --git a/Engine/Plugins/Experimental/AdvancedRenamer/Source/AdvancedRenamer/Private/Slate/SAdvancedRenamerPanel.cpp b/Engine/Plugins/Experimental/AdvancedRenamer/Source/AdvancedRenamer/Private/Slate/SAdvancedRenamerPanel.cpp
--- a/Engine/Plugins/Experimental/AdvancedRenamer/Source/AdvancedRenamer/Private/Slate/SAdvancedRenamerPanel.cpp
+++ b/Engine/Plugins/Experimental/AdvancedRenamer/Source/AdvancedRenamer/Private/Slate/SAdvancedRenamerPanel.cpp
@@ -95,7 +95,7 @@ void SAdvancedRenamerPanel::Construct(const FArguments& InArgs, const TSharedRef
 	Renamer = InRenamer;
 
 	SortMode = EColumnSortMode::None;
-	PreviewList = InRenamer->GetPreviews();
+	PreviewList = InRenamer->GetSortablePreviews();
 
 	CommandList = MakeShared<FUICommandList>();
 	CommandList->MapAction(
@@ -455,7 +455,7 @@ void SAdvancedRenamerPanel::OnColumnSortModeChanged(EColumnSortPriority::Type In
 
 	if (Renamer.IsValid())
 	{
-		PreviewList = Renamer->GetPreviews();
+		TArray<TSharedPtr<FAdvancedRenamerPreview>>& SortablePreview = Renamer->GetSortablePreviews();
 		if (SortMode != EColumnSortMode::None)
 		{
 			EColumnSortMode::Type TempSortMode = SortMode;
@@ -470,8 +470,16 @@ void SAdvancedRenamerPanel::OnColumnSortModeChanged(EColumnSortPriority::Type In
 				return TempSortMode == EColumnSortMode::Ascending ? CompareResult : !CompareResult;
 			};
 
-			PreviewList.Sort(ComparePreviewList);
+			SortablePreview.Sort(ComparePreviewList);
+			PreviewList = SortablePreview;
 		}
+		else
+		{
+			Renamer->ResetSortablePreviews();
+			PreviewList = Renamer->GetSortablePreviews();
+		}
+
+		Renamer->MarkDirty();
 		RenamePreviewList->RequestListRefresh();
 	}
 }
