Fix a few more menus that did not take into consideration bCanBeModified
#rnx

diff --git a/Engine/Plugins/Experimental/AdvancedRenamer/Source/AdvancedRenamer/Private/Integrations/AdvancedRenamerContentBrowserIntegration.cpp b/Engine/Plugins/Experimental/AdvancedRenamer/Source/AdvancedRenamer/Private/Integrations/AdvancedRenamerContentBrowserIntegration.cpp
--- a/Engine/Plugins/Experimental/AdvancedRenamer/Source/AdvancedRenamer/Private/Integrations/AdvancedRenamerContentBrowserIntegration.cpp
+++ b/Engine/Plugins/Experimental/AdvancedRenamer/Source/AdvancedRenamer/Private/Integrations/AdvancedRenamerContentBrowserIntegration.cpp
@@ -66,11 +66,15 @@ namespace UE::AdvancedRenamer::Private
 		{
 			FToolMenuOwnerScoped OwnerScoped(TEXT("AdvancedRenamer"));
 
-			if (UToolMenu* Menu = ToolMenus->ExtendMenu(TEXT("ContentBrowser.AssetContextMenu")))
+			UToolMenus::Get()->ExtendMenu("ContentBrowser.AssetContextMenu")->AddDynamicSection("BatchRenameDynamic", FNewToolMenuDelegate::CreateLambda([](UToolMenu* InMenu)
 			{
-				FToolMenuSection& MenuSection = Menu->FindOrAddSection(TEXT("CommonAssetActions"));
-				AddMenuEntry(MenuSection);
-			}
+				const UContentBrowserAssetContextMenuContext* Context = InMenu->FindContext<UContentBrowserAssetContextMenuContext>();
+				if (Context && Context->bCanBeModified)
+				{
+					FToolMenuSection& MenuSection = InMenu->FindOrAddSection(TEXT("CommonAssetActions"));
+					AddMenuEntry(MenuSection);
+				}
+			}));
 		}
 	}
 
