Add a way to disable opening the advanced renamer for actors that don't support renaming

diff --git a/Engine/Plugins/Experimental/AdvancedRenamer/Source/AdvancedRenamer/Private/AdvancedRenamerModule.cpp b/Engine/Plugins/Experimental/AdvancedRenamer/Source/AdvancedRenamer/Private/AdvancedRenamerModule.cpp
--- a/Engine/Plugins/Experimental/AdvancedRenamer/Source/AdvancedRenamer/Private/AdvancedRenamerModule.cpp
+++ b/Engine/Plugins/Experimental/AdvancedRenamer/Source/AdvancedRenamer/Private/AdvancedRenamerModule.cpp
@@ -88,6 +88,14 @@ void FAdvancedRenamerModule::OpenAdvancedRenamerForActors(const TArray<AActor*>&
 			return TWeakObjectPtr<AActor>(InActor);
 		}
 	);
+	
+	OnFilterAdvancedRenamerActors().Broadcast(WeakObjects);
+
+	// Don't open the advanced renamer if there are no actors remaining after the filter
+	if (WeakObjects.IsEmpty())
+	{
+		return;
+	}
 
 	TSharedRef<FAdvancedRenamerActorProvider> ActorProvider = MakeShared<FAdvancedRenamerActorProvider>();
 	ActorProvider->SetActorList(WeakObjects);
diff --git a/Engine/Plugins/Experimental/AdvancedRenamer/Source/AdvancedRenamer/Private/AdvancedRenamerModule.h b/Engine/Plugins/Experimental/AdvancedRenamer/Source/AdvancedRenamer/Private/AdvancedRenamerModule.h
--- a/Engine/Plugins/Experimental/AdvancedRenamer/Source/AdvancedRenamer/Private/AdvancedRenamerModule.h
+++ b/Engine/Plugins/Experimental/AdvancedRenamer/Source/AdvancedRenamer/Private/AdvancedRenamerModule.h
@@ -27,6 +27,7 @@ class FAdvancedRenamerModule : public IAdvancedRenamerModule
 	virtual void OpenAdvancedRenamer(const TSharedRef<IAdvancedRenamer>& InRenamer, const TSharedPtr<SWidget>& InParentWidget) override;
 	virtual TArray<AActor*> GetActorsSharingClassesInWorld(const TArray<AActor*>& InActors) override;
 	virtual TArray<TSharedPtr<IAdvancedRenamerSection>> GetRegisteredSections() const override { return Sections; }
+	virtual FOnFilterAdvancedRenamerActors& OnFilterAdvancedRenamerActors() override { return FilterAdvancedRenamerActors; }
 	//~ End IAdvancedRenamerModule
 
 private:
@@ -35,4 +36,5 @@ private:
 private:
 	TArray<TSharedPtr<IAdvancedRenamerSection>> Sections;
 	FDelegateHandle EnableRenamerHandle;
+	FOnFilterAdvancedRenamerActors FilterAdvancedRenamerActors;
 };
diff --git a/Engine/Plugins/Experimental/AdvancedRenamer/Source/AdvancedRenamer/Public/IAdvancedRenamerModule.h b/Engine/Plugins/Experimental/AdvancedRenamer/Source/AdvancedRenamer/Public/IAdvancedRenamerModule.h
--- a/Engine/Plugins/Experimental/AdvancedRenamer/Source/AdvancedRenamer/Public/IAdvancedRenamerModule.h
+++ b/Engine/Plugins/Experimental/AdvancedRenamer/Source/AdvancedRenamer/Public/IAdvancedRenamerModule.h
@@ -47,4 +47,9 @@ public:
 	virtual TArray<TSharedPtr<IAdvancedRenamerSection>> GetRegisteredSections() const = 0;
 
 	virtual TArray<AActor*> GetActorsSharingClassesInWorld(const TArray<AActor*>& InActors) = 0;
+
+	// Delegate to filter out actors before they are opened in the advanced renamer. The renamer will not be opened if the array is empty
+	DECLARE_MULTICAST_DELEGATE_OneParam(FOnFilterAdvancedRenamerActors, TArray<TWeakObjectPtr<AActor>>&);
+	virtual FOnFilterAdvancedRenamerActors& OnFilterAdvancedRenamerActors() = 0;
+	
 };
